Sanity check: nullptr check on FOutputBuffer::Init, to ensure that the IAudioMixerPlatformInterface hasn't nulled out the IAudioMixer ptr during OpenAudioStream.

diff --git a/Engine/Source/Runtime/AudioMixerCore/Private/AudioMixer.cpp b/Engine/Source/Runtime/AudioMixerCore/Private/AudioMixer.cpp
--- a/Engine/Source/Runtime/AudioMixerCore/Private/AudioMixer.cpp
+++ b/Engine/Source/Runtime/AudioMixerCore/Private/AudioMixer.cpp
@@ -119,6 +119,8 @@ namespace Audio
 	{
 		Buffer.SetNumZeroed(InNumSamples);
 		DataFormat = InDataFormat;
+
+		check(InAudioMixer != nullptr);
 		AudioMixer = InAudioMixer;
 
 		switch (DataFormat)
@@ -144,7 +146,10 @@ namespace Audio
 
 		// Zero the buffer
 		FPlatformMemory::Memzero(Buffer.GetData(), Buffer.Num() * sizeof(float));
-		AudioMixer->OnProcessAudioStream(Buffer);
+		if (AudioMixer != nullptr)
+		{
+			AudioMixer->OnProcessAudioStream(Buffer);
+		}
 
 		switch (DataFormat)
 		{
