Added missing include.

diff --git a/Engine/Source/Runtime/AudioMixerCore/Private/AudioMixerNullDevice.cpp b/Engine/Source/Runtime/AudioMixerCore/Private/AudioMixerNullDevice.cpp
--- a/Engine/Source/Runtime/AudioMixerCore/Private/AudioMixerNullDevice.cpp
+++ b/Engine/Source/Runtime/AudioMixerCore/Private/AudioMixerNullDevice.cpp
@@ -3,6 +3,7 @@
 #include "AudioMixerNullDevice.h"
 #include "CoreMinimal.h"
 #include "HAL/PlatformProcess.h"
+#include "HAL/PlatformTime.h"
 
 namespace Audio
 {
