Make sure AudioFadeEvent is valid before using it

When stopping StopAudioStream is called it sets AudioFadeEvent to nullptr and if you exit after that the teardown will hit this nullptr

diff --git a/Engine/Source/Runtime/AudioMixerCore/Private/AudioMixer.cpp b/Engine/Source/Runtime/AudioMixerCore/Private/AudioMixer.cpp
--- a/Engine/Source/Runtime/AudioMixerCore/Private/AudioMixer.cpp
+++ b/Engine/Source/Runtime/AudioMixerCore/Private/AudioMixer.cpp
@@ -274,7 +274,11 @@ namespace Audio
 		}
 
 		bPerformingFade = true;
-		AudioFadeEvent->Wait();
+		if (AudioFadeEvent != nullptr)
+		{
+			AudioFadeEvent->Wait();
+		}
+
 		FadeVolume = 0.0f;
 	}
 
