When not using threading avoid checking the audio thread has stopped, and set the state to stopped. As that thread is not running

diff --git a/Engine/Source/Runtime/AudioMixerCore/Private/AudioMixer.cpp b/Engine/Source/Runtime/AudioMixerCore/Private/AudioMixer.cpp
--- a/Engine/Source/Runtime/AudioMixerCore/Private/AudioMixer.cpp
+++ b/Engine/Source/Runtime/AudioMixerCore/Private/AudioMixer.cpp
@@ -501,7 +501,16 @@ namespace Audio
 		if (AudioRenderThread != nullptr)
 		{
 			AudioRenderThread->WaitForCompletion();
-			check(AudioStreamInfo.StreamState == EAudioOutputStreamState::Stopped);
+
+			// WaitForCompletion will complete right away when single threaded, and AudioStreamInfo.StreamState will never be set to stopped
+			if (FPlatformProcess::SupportsMultithreading())
+			{
+				check(AudioStreamInfo.StreamState == EAudioOutputStreamState::Stopped);
+			}
+			else
+			{
+				AudioStreamInfo.StreamState = EAudioOutputStreamState::Stopped;
+			}
 
 			delete AudioRenderThread;
 			AudioRenderThread = nullptr;
