Added missing include

diff --git a/Engine/Source/Runtime/AudioMixerCore/Public/AudioMixer.h b/Engine/Source/Runtime/AudioMixerCore/Public/AudioMixer.h
--- a/Engine/Source/Runtime/AudioMixerCore/Public/AudioMixer.h
+++ b/Engine/Source/Runtime/AudioMixerCore/Public/AudioMixer.h
@@ -7,6 +7,7 @@
 #include "AudioMixerTypes.h"
 #include "HAL/Runnable.h"
 #include "HAL/ThreadSafeBool.h"
+#include "Misc/ScopeLock.h"
 #include "Misc/SingleThreadRunnable.h"
 #include "AudioMixerNullDevice.h"
 #include "DSP/ParamInterpolator.h"
