Adding optional config to wait for audio to be rendered before falling back to submitting an underrun (all zeros) buffer.

diff --git a/Engine/Source/Runtime/AudioMixerCore/Private/AudioMixer.cpp b/Engine/Source/Runtime/AudioMixerCore/Private/AudioMixer.cpp
--- a/Engine/Source/Runtime/AudioMixerCore/Private/AudioMixer.cpp
+++ b/Engine/Source/Runtime/AudioMixerCore/Private/AudioMixer.cpp
@@ -57,6 +57,14 @@ FAutoConsoleVariableRef CVarOverrunTimeout(
 	TEXT("Amount of time to wait for the render thread to time out before swapping to the null device. \n"),
 	ECVF_Default);
 
+static int32 UnderrunTimeoutCVar = 0;
+FAutoConsoleVariableRef CVarUnderrunTimeout(
+	TEXT("au.UnderrunTimeoutMSec"),
+	UnderrunTimeoutCVar,
+	TEXT("Amount of time to wait for the render thread to generate the next buffer before submitting an underrun buffer. \n"),
+	ECVF_Default);
+
+
 namespace Audio
 {
 	int32 sRenderInstanceIds = 0;
@@ -123,6 +131,12 @@ namespace Audio
 		check(InAudioMixer != nullptr);
 		AudioMixer = InAudioMixer;
 
+		if (IsReadyEvent == nullptr)
+		{
+			IsReadyEvent = FPlatformProcess::GetSynchEventFromPool(true /*Manual Reset*/);
+		}
+		check(IsReadyEvent != nullptr);
+
 		switch (DataFormat)
 		{
 			case EAudioMixerStreamDataFormat::Float:
@@ -174,8 +188,9 @@ namespace Audio
 			break;
 		}
 
-		// Mark that we're ready
+		// Mark/signal that we're ready
 		bIsReady = true;
+		IsReadyEvent->Trigger();
  	}
  
 	const uint8* FOutputBuffer::GetBufferData() const
@@ -405,16 +420,27 @@ namespace Audio
 		// If it's not ready, warn, and then wait here. This will cause underruns but is preferable than getting out-of-order buffer state.
 		static int32 UnderrunCount = 0;
 		static int32 CurrentUnderrunCount = 0;
-		
+
+		bool bSubmittingUnderrunBuffer = false;
+
 		if (!OutputBuffers[NextReadIndex].IsReady())
 		{
-
+			// try to wait for the buffer to be ready
+			FEvent* BufferReadyEvent = OutputBuffers[NextReadIndex].IsReadyEvent;
+			if (!BufferReadyEvent || !BufferReadyEvent->Wait(static_cast<uint32>(UnderrunTimeoutCVar)))
+			{
+				bSubmittingUnderrunBuffer = true; // Event didn't fire in time
+			}
+		}
+		
+		if (bSubmittingUnderrunBuffer)
+		{
 			UnderrunCount++;
 			CurrentUnderrunCount++;
 			
 			if (!bWarnedBufferUnderrun)
 			{						
-				UE_LOG(LogAudioMixerDebug, Log, TEXT("Audio Buffer Underrun detected."));
+				UE_LOG(LogAudioMixer, Display, TEXT("Audio Buffer Underrun detected."));
 				bWarnedBufferUnderrun = true;
 			}
 		
@@ -437,6 +463,7 @@ namespace Audio
 
 			// Update the current read index to the next read index
 			CurrentBufferReadIndex = NextReadIndex;
+			OutputBuffers[NextReadIndex].IsReadyEvent->Reset();
 		}
 
 		DeviceSwapCriticalSection.Unlock();
diff --git a/Engine/Source/Runtime/AudioMixerCore/Public/AudioMixer.h b/Engine/Source/Runtime/AudioMixerCore/Public/AudioMixer.h
--- a/Engine/Source/Runtime/AudioMixerCore/Public/AudioMixer.h
+++ b/Engine/Source/Runtime/AudioMixerCore/Public/AudioMixer.h
@@ -290,9 +290,19 @@ namespace Audio
 	{
 	public:
 		FOutputBuffer()
-			: AudioMixer(nullptr)
+			: IsReadyEvent(nullptr)
+			, AudioMixer(nullptr)
 			, DataFormat(EAudioMixerStreamDataFormat::Unknown)
 		{}
+
+		~FOutputBuffer()
+		{
+			if (IsReadyEvent != nullptr)
+			{
+				FPlatformProcess::ReturnSynchEventToPool(IsReadyEvent);
+				IsReadyEvent = nullptr;
+			}
+		}
  
 		/** Initialize the buffer with the given samples and output format. */
 		void Init(IAudioMixer* InAudioMixer, const int32 InNumSamples, const EAudioMixerStreamDataFormat::Type InDataFormat);
@@ -308,6 +318,9 @@ namespace Audio
 		const uint8* GetBufferData() const;
 		uint8* GetBufferData();
 
+		/** Event to signal that the buffer is ready */
+		FEvent* IsReadyEvent;
+
 		/** Gets the number of frames of the buffer. */
 		int32 GetNumFrames() const;
 
@@ -318,7 +331,14 @@ namespace Audio
 		bool IsReady() const { return bIsReady; }
 
 		/** Resets the buffer ready state. */
-		void ResetReadyState() { bIsReady = false; }
+		void ResetReadyState()
+		{
+			bIsReady = false;
+			if (IsReadyEvent)
+			{
+				IsReadyEvent->Reset();
+			}
+		}
 
 		/** Resets the internal buffers to the new sample count. Used when device is changed. */
 		void Reset(const int32 InNewNumSamples);
@@ -559,10 +579,10 @@ namespace Audio
 		FEvent* AudioFadeEvent;
 
 		/** The buffer which is currently submitted to the output device (and is being read from). */
-		int32 CurrentBufferReadIndex;
+		TAtomic<int32> CurrentBufferReadIndex;
 
 		/** The buffer which is currently being rendered to (or about to be rendered to). */
-		int32 CurrentBufferWriteIndex;
+		TAtomic<int32> CurrentBufferWriteIndex;
 
 		/** The number of mixer buffers to queue on the output source voice. */
 		int32 NumOutputBuffers;
