Fix non-unity errors
#rnx

diff --git a/Engine/Source/Runtime/AudioMixerCore/Private/AudioMixer.cpp b/Engine/Source/Runtime/AudioMixerCore/Private/AudioMixer.cpp
--- a/Engine/Source/Runtime/AudioMixerCore/Private/AudioMixer.cpp
+++ b/Engine/Source/Runtime/AudioMixerCore/Private/AudioMixer.cpp
@@ -123,6 +123,15 @@ namespace Audio
 		}
 	}
 
+	FOutputBuffer::~FOutputBuffer()
+	{
+		if (IsReadyEvent != nullptr)
+		{
+			FPlatformProcess::ReturnSynchEventToPool(IsReadyEvent);
+			IsReadyEvent = nullptr;
+		}
+	}
+
 	void FOutputBuffer::Init(IAudioMixer* InAudioMixer, const int32 InNumSamples, const EAudioMixerStreamDataFormat::Type InDataFormat)
 	{
 		Buffer.SetNumZeroed(InNumSamples);
@@ -222,6 +231,15 @@ namespace Audio
 		return Buffer.Num();
 	}
 
+	void FOutputBuffer::ResetReadyState()
+	{
+		bIsReady = false;
+		if (IsReadyEvent)
+		{
+			IsReadyEvent->Reset();
+		}
+	}
+
 	void FOutputBuffer::Reset(const int32 InNewNumSamples)
 	{
 		Buffer.Reset();
diff --git a/Engine/Source/Runtime/AudioMixerCore/Public/AudioMixer.h b/Engine/Source/Runtime/AudioMixerCore/Public/AudioMixer.h
--- a/Engine/Source/Runtime/AudioMixerCore/Public/AudioMixer.h
+++ b/Engine/Source/Runtime/AudioMixerCore/Public/AudioMixer.h
@@ -295,14 +295,7 @@ namespace Audio
 			, DataFormat(EAudioMixerStreamDataFormat::Unknown)
 		{}
 
-		~FOutputBuffer()
-		{
-			if (IsReadyEvent != nullptr)
-			{
-				FPlatformProcess::ReturnSynchEventToPool(IsReadyEvent);
-				IsReadyEvent = nullptr;
-			}
-		}
+		~FOutputBuffer();
  
 		/** Initialize the buffer with the given samples and output format. */
 		void Init(IAudioMixer* InAudioMixer, const int32 InNumSamples, const EAudioMixerStreamDataFormat::Type InDataFormat);
@@ -331,14 +324,7 @@ namespace Audio
 		bool IsReady() const { return bIsReady; }
 
 		/** Resets the buffer ready state. */
-		void ResetReadyState()
-		{
-			bIsReady = false;
-			if (IsReadyEvent)
-			{
-				IsReadyEvent->Reset();
-			}
-		}
+		void ResetReadyState();
 
 		/** Resets the internal buffers to the new sample count. Used when device is changed. */
 		void Reset(const int32 InNewNumSamples);
