Edigrate of audio underflow timeout option (defaulting to 5ms)

diff --git a/Engine/Source/Runtime/AudioMixerCore/Private/AudioMixer.cpp b/Engine/Source/Runtime/AudioMixerCore/Private/AudioMixer.cpp
--- a/Engine/Source/Runtime/AudioMixerCore/Private/AudioMixer.cpp
+++ b/Engine/Source/Runtime/AudioMixerCore/Private/AudioMixer.cpp
@@ -57,7 +57,7 @@ FAutoConsoleVariableRef CVarOverrunTimeout(
 	TEXT("Amount of time to wait for the render thread to time out before swapping to the null device. \n"),
 	ECVF_Default);
 
-static int32 UnderrunTimeoutCVar = 0;
+static int32 UnderrunTimeoutCVar = 5;
 FAutoConsoleVariableRef CVarUnderrunTimeout(
 	TEXT("au.UnderrunTimeoutMSec"),
 	UnderrunTimeoutCVar,
@@ -132,6 +132,7 @@ namespace Audio
 		}
 	}
 
+
 	void FOutputBuffer::Init(IAudioMixer* InAudioMixer, const int32 InNumSamples, const EAudioMixerStreamDataFormat::Type InDataFormat)
 	{
 		Buffer.SetNumZeroed(InNumSamples);
@@ -240,6 +241,7 @@ namespace Audio
 		}
 	}
 
+
 	void FOutputBuffer::Reset(const int32 InNewNumSamples)
 	{
 		Buffer.Reset();
