Fix for race condition between audio render thread and device callback on startup causing a crash.

diff --git a/Engine/Source/Runtime/AudioMixerCore/Private/AudioMixer.cpp b/Engine/Source/Runtime/AudioMixerCore/Private/AudioMixer.cpp
--- a/Engine/Source/Runtime/AudioMixerCore/Private/AudioMixer.cpp
+++ b/Engine/Source/Runtime/AudioMixerCore/Private/AudioMixer.cpp
@@ -473,8 +473,12 @@ namespace Audio
 			return;
 		}
 
-		check(CurrentBufferReadIndex != INDEX_NONE);
-		check(CurrentBufferWriteIndex != INDEX_NONE);
+		// AudioRenderThread hasn't executed yet, return silence
+		if (CurrentBufferReadIndex == INDEX_NONE || CurrentBufferWriteIndex == INDEX_NONE)
+		{
+			SubmitBuffer(UnderrunBuffer.GetBufferData());
+			return;
+		}
 
 		// Reset the ready state of the buffer which was just finished playing
 		FOutputBuffer& CurrentReadBuffer = OutputBuffers[CurrentBufferReadIndex];
@@ -548,9 +552,6 @@ namespace Audio
 		// Set the number of buffers to be one more than the number to queue.
 		NumOutputBuffers = FMath::Max(OpenStreamParams.NumBuffers, 2);
 
-		CurrentBufferReadIndex = 0;
-		CurrentBufferWriteIndex = 1;
-
 		OutputBuffers.Reset();
 		OutputBuffers.AddDefaulted(NumOutputBuffers);
 		for (int32 Index = 0; Index < NumOutputBuffers; ++Index)
@@ -652,12 +653,14 @@ namespace Audio
 		// Lets prime and submit the first buffer (which is going to be the buffer underrun buffer)
 		SubmitBuffer(UnderrunBuffer.GetBufferData());
 
-		OutputBuffers[CurrentBufferWriteIndex].MixNextBuffer();
-
-		check(CurrentBufferReadIndex == 0);
-		check(CurrentBufferWriteIndex == 1);
+		OutputBuffers[0].MixNextBuffer();
 
 		// Start immediately processing the next buffer
+		check(CurrentBufferReadIndex == INDEX_NONE);
+		check(CurrentBufferWriteIndex == INDEX_NONE);
+
+		CurrentBufferReadIndex = 0;
+		CurrentBufferWriteIndex = 1;
 
 		while (AudioStreamInfo.StreamState != EAudioOutputStreamState::Stopping)
 		{
