Fix for assert being hit in automated test. (resetting read/write indices)

diff --git a/Engine/Source/Runtime/AudioMixerCore/Private/AudioMixer.cpp b/Engine/Source/Runtime/AudioMixerCore/Private/AudioMixer.cpp
--- a/Engine/Source/Runtime/AudioMixerCore/Private/AudioMixer.cpp
+++ b/Engine/Source/Runtime/AudioMixerCore/Private/AudioMixer.cpp
@@ -423,6 +423,11 @@ namespace Audio
 
 	void IAudioMixerPlatformInterface::StopRunningNullDevice()
 	{
+		if (bIsUsingNullDevice)
+		{
+			CurrentBufferReadIndex = INDEX_NONE;
+			CurrentBufferWriteIndex = INDEX_NONE;
+		}
 		if (NullDeviceCallback.IsValid())
 		{
 			NullDeviceCallback.Reset();
@@ -688,6 +693,8 @@ namespace Audio
 			}
 		}
 
+		CurrentBufferReadIndex = INDEX_NONE;
+		CurrentBufferWriteIndex = INDEX_NONE;
 		OpenStreamParams.AudioMixer->OnAudioStreamShutdown();
 
 		AudioStreamInfo.StreamState = EAudioOutputStreamState::Stopped;
