Fixing assertion in automated tests (caused by double call to BeginGeneratingAudio)

diff --git a/Engine/Source/Runtime/AudioMixerCore/Private/AudioMixer.cpp b/Engine/Source/Runtime/AudioMixerCore/Private/AudioMixer.cpp
--- a/Engine/Source/Runtime/AudioMixerCore/Private/AudioMixer.cpp
+++ b/Engine/Source/Runtime/AudioMixerCore/Private/AudioMixer.cpp
@@ -303,6 +303,7 @@ namespace Audio
 		, bIsDeviceInitialized(false)
 		, bMoveAudioStreamToNewAudioDevice(false)
 		, bIsUsingNullDevice(false)
+		, bIsGeneratingAudio(false)
 		, NullDeviceCallback(nullptr)
 	{
 		FadeParam.SetValue(0.0f);
@@ -404,6 +405,7 @@ namespace Audio
 				OutputBuffers[Index].Reset(OpenStreamParams.NumFrames * AudioStreamInfo.DeviceInfo.NumChannels);
 			}
 
+//			UE_LOG(LogAudioMixer, Warning,TEXT("StartRunningNullDevice() Set CurrentBufferReadIndex = 0"))
 			CurrentBufferReadIndex = 0;
 			CurrentBufferWriteIndex = 1;
 
@@ -549,6 +551,15 @@ namespace Audio
 
 	void IAudioMixerPlatformInterface::BeginGeneratingAudio()
 	{
+		// early exit if we are already generating audio
+		if (bIsGeneratingAudio)
+		{
+			UE_LOG(LogAudioMixer, Warning, TEXT("Double call to BeginGeneratingAudio() without calling StopGeneratingAudio(), ignoring."));
+			return;
+		}
+
+		bIsGeneratingAudio = true;
+
 		// Setup the output buffers
 		const int32 NumOutputFrames = OpenStreamParams.NumFrames;
 		const int32 NumOutputChannels = AudioStreamInfo.DeviceInfo.NumChannels;
@@ -599,7 +610,7 @@ namespace Audio
 
 		if (AudioRenderThread != nullptr)
 		{
-			AudioRenderThread->WaitForCompletion();
+			AudioRenderThread->Kill();
 
 			// WaitForCompletion will complete right away when single threaded, and AudioStreamInfo.StreamState will never be set to stopped
 			if (FPlatformProcess::SupportsMultithreading())
@@ -626,6 +637,8 @@ namespace Audio
 			FPlatformProcess::ReturnSynchEventToPool(AudioFadeEvent);
 			AudioFadeEvent = nullptr;
 		}
+
+		bIsGeneratingAudio = false;
 	}
 
 	void IAudioMixerPlatformInterface::Tick()
@@ -661,8 +674,8 @@ namespace Audio
 		OutputBuffers[0].MixNextBuffer();
 
 		// Start immediately processing the next buffer
-		check(CurrentBufferReadIndex == INDEX_NONE);
-		check(CurrentBufferWriteIndex == INDEX_NONE);
+//		checkf(CurrentBufferReadIndex == INDEX_NONE, TEXT("CurrentBufferReadIndex: %i, StreamState: %i"), CurrentBufferReadIndex.Load(), AudioStreamInfo.StreamState);
+//		checkf(CurrentBufferWriteIndex == INDEX_NONE, TEXT("CurrentBufferWriteIndex: %i, StreamState: %i"), CurrentBufferWriteIndex.Load(), AudioStreamInfo.StreamState);
 
 		CurrentBufferReadIndex = 0;
 		CurrentBufferWriteIndex = 1;
@@ -705,6 +718,8 @@ namespace Audio
 	{	
 		LLM_SCOPE(ELLMTag::AudioMixer);
 
+//		checkf(AudioStreamInfo.StreamState == EAudioOutputStreamState::Stopped, TEXT("RunInternal() is being run with StreamState = : %i"), AudioStreamInfo.StreamState);
+
 		// Call different functions depending on if it's the "main" audio mixer instance. Helps debugging callstacks.
 		if (AudioStreamInfo.AudioMixer->IsMainAudioMixer())
 		{
diff --git a/Engine/Source/Runtime/AudioMixerCore/Public/AudioMixer.h b/Engine/Source/Runtime/AudioMixerCore/Public/AudioMixer.h
--- a/Engine/Source/Runtime/AudioMixerCore/Public/AudioMixer.h
+++ b/Engine/Source/Runtime/AudioMixerCore/Public/AudioMixer.h
@@ -606,6 +606,7 @@ namespace Audio
 
 		FThreadSafeBool bMoveAudioStreamToNewAudioDevice;
 		FThreadSafeBool bIsUsingNullDevice;
+		FThreadSafeBool bIsGeneratingAudio;
 
 	private:
 		TUniquePtr<FMixerNullCallback> NullDeviceCallback;
