Fixing up double initialization of Next Read/Write Buffer indices for null device audio rendering (was causing automated tests to hit assert)

diff --git a/Engine/Source/Runtime/AudioMixerCore/Private/AudioMixer.cpp b/Engine/Source/Runtime/AudioMixerCore/Private/AudioMixer.cpp
--- a/Engine/Source/Runtime/AudioMixerCore/Private/AudioMixer.cpp
+++ b/Engine/Source/Runtime/AudioMixerCore/Private/AudioMixer.cpp
@@ -405,12 +405,7 @@ namespace Audio
 				OutputBuffers[Index].Reset(OpenStreamParams.NumFrames * AudioStreamInfo.DeviceInfo.NumChannels);
 			}
 
-//			UE_LOG(LogAudioMixer, Warning,TEXT("StartRunningNullDevice() Set CurrentBufferReadIndex = 0"))
-			CurrentBufferReadIndex = 0;
-			CurrentBufferWriteIndex = 1;
-
-			SubmitBuffer(OutputBuffers[CurrentBufferReadIndex].GetBufferData());
-			check(OpenStreamParams.NumFrames * AudioStreamInfo.DeviceInfo.NumChannels == OutputBuffers[CurrentBufferReadIndex].GetBuffer().Num());
+			check(OpenStreamParams.NumFrames * AudioStreamInfo.DeviceInfo.NumChannels == OutputBuffers[0].GetBuffer().Num());
 
 			AudioRenderEvent->Trigger();
 
@@ -551,12 +546,7 @@ namespace Audio
 
 	void IAudioMixerPlatformInterface::BeginGeneratingAudio()
 	{
-		// early exit if we are already generating audio
-		if (bIsGeneratingAudio)
-		{
-			UE_LOG(LogAudioMixer, Warning, TEXT("Double call to BeginGeneratingAudio() without calling StopGeneratingAudio(), ignoring."));
-			return;
-		}
+		checkf(!bIsGeneratingAudio, TEXT("BeginGeneratingAudio() is being run with StreamState = %i and bIsGeneratingAudio = %i"), AudioStreamInfo.StreamState, !!bIsGeneratingAudio);
 
 		bIsGeneratingAudio = true;
 
@@ -674,8 +664,8 @@ namespace Audio
 		OutputBuffers[0].MixNextBuffer();
 
 		// Start immediately processing the next buffer
-//		checkf(CurrentBufferReadIndex == INDEX_NONE, TEXT("CurrentBufferReadIndex: %i, StreamState: %i"), CurrentBufferReadIndex.Load(), AudioStreamInfo.StreamState);
-//		checkf(CurrentBufferWriteIndex == INDEX_NONE, TEXT("CurrentBufferWriteIndex: %i, StreamState: %i"), CurrentBufferWriteIndex.Load(), AudioStreamInfo.StreamState);
+		checkf(CurrentBufferReadIndex == INDEX_NONE, TEXT("CurrentBufferReadIndex: %i, StreamState: %i"), CurrentBufferReadIndex.Load(), AudioStreamInfo.StreamState);
+		checkf(CurrentBufferWriteIndex == INDEX_NONE, TEXT("CurrentBufferWriteIndex: %i, StreamState: %i"), CurrentBufferWriteIndex.Load(), AudioStreamInfo.StreamState);
 
 		CurrentBufferReadIndex = 0;
 		CurrentBufferWriteIndex = 1;
@@ -718,8 +708,6 @@ namespace Audio
 	{	
 		LLM_SCOPE(ELLMTag::AudioMixer);
 
-//		checkf(AudioStreamInfo.StreamState == EAudioOutputStreamState::Stopped, TEXT("RunInternal() is being run with StreamState = : %i"), AudioStreamInfo.StreamState);
-
 		// Call different functions depending on if it's the "main" audio mixer instance. Helps debugging callstacks.
 		if (AudioStreamInfo.AudioMixer->IsMainAudioMixer())
 		{
diff --git a/Engine/Source/Runtime/NonRealtimeAudioRenderer/Private/AudioMixerPlatformNonRealtime.cpp b/Engine/Source/Runtime/NonRealtimeAudioRenderer/Private/AudioMixerPlatformNonRealtime.cpp
--- a/Engine/Source/Runtime/NonRealtimeAudioRenderer/Private/AudioMixerPlatformNonRealtime.cpp
+++ b/Engine/Source/Runtime/NonRealtimeAudioRenderer/Private/AudioMixerPlatformNonRealtime.cpp
@@ -66,6 +66,10 @@ namespace Audio
 		double SecondsRendered = TotalDurationRendered;
 		TotalDurationRendered += NumSecondsToRender;
 
+		CurrentBufferWriteIndex = 0;
+		CurrentBufferReadIndex = 0;
+
+
 		while (SecondsRendered < TotalDurationRendered)
 		{
 			// RenderTimeAnalysis.Start();
@@ -75,6 +79,9 @@ namespace Audio
 			ReadNextBuffer();
 			SecondsRendered += TimePerCallback;
 		}
+
+		CurrentBufferReadIndex = INDEX_NONE;
+		CurrentBufferWriteIndex = INDEX_NONE;
 	}
 
 	void FMixerPlatformNonRealtime::OpenFileToWriteAudioTo(const FString& OutPath)
