Missing file

diff --git a/Engine/Source/Runtime/AudioMixerCore/Public/AudioMixer.h b/Engine/Source/Runtime/AudioMixerCore/Public/AudioMixer.h
--- a/Engine/Source/Runtime/AudioMixerCore/Public/AudioMixer.h
+++ b/Engine/Source/Runtime/AudioMixerCore/Public/AudioMixer.h
@@ -67,7 +67,13 @@ namespace EAudioMixerChannel
 		TopBackCenter,
 		TopBackRight,
 		Unknown,
-		ChannelTypeCount
+		ChannelTypeCount,
+
+		// Alias Values
+		Ambisonics_W = FrontLeft,
+		Ambisonics_X = FrontRight,
+		Ambisonics_Y = FrontCenter,
+		Ambisonics_Z = LowFrequency,
 	};
 
 	static const int32 MaxSupportedChannel = EAudioMixerChannel::TopCenter;
