- Fix for issue when audio engine goes silent when swaping from null device
- Fix null ref on shutdown

diff --git a/Engine/Source/Runtime/AudioMixerCore/Private/AudioMixer.cpp b/Engine/Source/Runtime/AudioMixerCore/Private/AudioMixer.cpp
--- a/Engine/Source/Runtime/AudioMixerCore/Private/AudioMixer.cpp
+++ b/Engine/Source/Runtime/AudioMixerCore/Private/AudioMixer.cpp
@@ -422,8 +422,8 @@ namespace Audio
 	{
 		if (bIsUsingNullDevice)
 		{
-			CurrentBufferReadIndex = INDEX_NONE;
-			CurrentBufferWriteIndex = INDEX_NONE;
+			CurrentBufferReadIndex = 0;
+			CurrentBufferWriteIndex = 1;
 		}
 		if (NullDeviceCallback.IsValid())
 		{
@@ -541,7 +541,10 @@ namespace Audio
 		DeviceSwapCriticalSection.Unlock();
 
 		// Kick off rendering of the next set of buffers
-		AudioRenderEvent->Trigger();
+		if (AudioRenderEvent)
+		{
+			AudioRenderEvent->Trigger();
+		}
 	}
 
 	void IAudioMixerPlatformInterface::BeginGeneratingAudio()
