Remove unused function header
#rnx

diff --git a/Engine/Source/Runtime/AudioMixerCore/Public/AudioMixer.h b/Engine/Source/Runtime/AudioMixerCore/Public/AudioMixer.h
--- a/Engine/Source/Runtime/AudioMixerCore/Public/AudioMixer.h
+++ b/Engine/Source/Runtime/AudioMixerCore/Public/AudioMixer.h
@@ -499,9 +499,6 @@ namespace Audio
 		/** Start a fadeout. Prevents pops during shutdown. */
 		void FadeOut();
 
-		/** Sets the mater volume of the audio device. This attenuates all audio, used for muting, etc. */
-		void SetMasterVolume(const float InVolume);
-
 		/** Returns the last error generated. */
 		FString GetLastError() const { return LastError; }
 
