#UE4 Fix deadlock at startup due to locked critical section that was never unlocked

diff --git a/Engine/Source/Runtime/AudioMixerCore/Private/AudioMixer.cpp b/Engine/Source/Runtime/AudioMixerCore/Private/AudioMixer.cpp
--- a/Engine/Source/Runtime/AudioMixerCore/Private/AudioMixer.cpp
+++ b/Engine/Source/Runtime/AudioMixerCore/Private/AudioMixer.cpp
@@ -479,6 +479,7 @@ namespace Audio
 		if (CurrentBufferReadIndex == INDEX_NONE || CurrentBufferWriteIndex == INDEX_NONE)
 		{
 			SubmitBuffer(UnderrunBuffer.GetBufferData());
+			DeviceSwapCriticalSection.Unlock();
 			return;
 		}
 
