Fix for Audio hang on exit

diff --git a/Engine/Source/Runtime/AudioMixerCore/Private/AudioMixer.cpp b/Engine/Source/Runtime/AudioMixerCore/Private/AudioMixer.cpp
--- a/Engine/Source/Runtime/AudioMixerCore/Private/AudioMixer.cpp
+++ b/Engine/Source/Runtime/AudioMixerCore/Private/AudioMixer.cpp
@@ -65,6 +65,13 @@ FAutoConsoleVariableRef CVarUnderrunTimeout(
 	TEXT("Amount of time to wait for the render thread to generate the next buffer before submitting an underrun buffer. \n"),
 	ECVF_Default);
 
+static int32 FadeoutTimeoutCVar = 2000;
+FAutoConsoleVariableRef CVarFadeoutTimeout(
+	TEXT("au.FadeOutTimeoutMSec"),
+	FadeoutTimeoutCVar,
+	TEXT("Amount of time to wait for the FadeOut Event to fire. \n"),
+	ECVF_Default);
+
 static float LinearGainScalarForFinalOututCVar = 1.0f;
 FAutoConsoleVariableRef LinearGainScalarForFinalOutut(
 	TEXT("au.LinearGainScalarForFinalOutut"),
@@ -344,8 +351,11 @@ namespace Audio
 
 		bPerformingFade = true;
 		if (AudioFadeEvent != nullptr)
-		{
-			AudioFadeEvent->Wait(1000);
+		{						
+			if (!AudioFadeEvent->Wait(FadeoutTimeoutCVar))
+			{
+				UE_LOG(LogAudioMixer, Warning, TEXT("FadeOutEvent timed out"));
+			}
 		}
 
 		FadeVolume = 0.0f;
diff --git a/Engine/Source/Runtime/AudioMixerCore/Public/AudioMixer.h b/Engine/Source/Runtime/AudioMixerCore/Public/AudioMixer.h
--- a/Engine/Source/Runtime/AudioMixerCore/Public/AudioMixer.h
+++ b/Engine/Source/Runtime/AudioMixerCore/Public/AudioMixer.h
@@ -489,10 +489,10 @@ namespace Audio
 		void ReadNextBuffer();
 
 		/** Reset the fade state (use if reusing audio platform interface, e.g. in main audio device. */
-		void FadeIn();
+		virtual void FadeIn();
 
 		/** Start a fadeout. Prevents pops during shutdown. */
-		void FadeOut();
+		virtual void FadeOut();
 
 		/** Returns the last error generated. */
 		FString GetLastError() const { return LastError; }
diff --git a/Engine/Source/Runtime/NonRealtimeAudioRenderer/Public/AudioMixerPlatformNonRealtime.h b/Engine/Source/Runtime/NonRealtimeAudioRenderer/Public/AudioMixerPlatformNonRealtime.h
--- a/Engine/Source/Runtime/NonRealtimeAudioRenderer/Public/AudioMixerPlatformNonRealtime.h
+++ b/Engine/Source/Runtime/NonRealtimeAudioRenderer/Public/AudioMixerPlatformNonRealtime.h
@@ -46,6 +46,8 @@ namespace Audio
 		virtual FAudioPlatformSettings GetPlatformSettings() const override;
 		virtual void OnHardwareUpdate() override;
 		virtual bool IsNonRealtime() const override;
+		virtual void FadeOut() override {}
+		virtual void FadeIn() override {}
 		//~ End IAudioMixerPlatformInterface
 
 	private:
