Added extra logging behind a cvar to help diagnose "client sometimes launches with no audio" issue

diff --git a/Engine/Source/Runtime/AudioMixer/Private/AudioMixerDevice.cpp b/Engine/Source/Runtime/AudioMixer/Private/AudioMixerDevice.cpp
--- a/Engine/Source/Runtime/AudioMixer/Private/AudioMixerDevice.cpp
+++ b/Engine/Source/Runtime/AudioMixer/Private/AudioMixerDevice.cpp
@@ -340,6 +340,11 @@ namespace Audio
 				return AudioMixerPlatform->StartAudioStream();
 			}
 		}
+		else if (AudioMixerPlatform)
+		{
+			UE_LOG(LogAudioMixer, Warning, TEXT("Failed to Open audio stream"));
+		}
+
 		return false;
 	}
 
diff --git a/Engine/Source/Runtime/AudioMixerCore/Private/AudioMixer.cpp b/Engine/Source/Runtime/AudioMixerCore/Private/AudioMixer.cpp
--- a/Engine/Source/Runtime/AudioMixerCore/Private/AudioMixer.cpp
+++ b/Engine/Source/Runtime/AudioMixerCore/Private/AudioMixer.cpp
@@ -80,6 +80,14 @@ FAutoConsoleVariableRef LinearGainScalarForFinalOutut(
 	TEXT("Default is 1.0f \n"),
 	ECVF_Default);
 
+static int32 ExtraAudioMixerDeviceLoggingCVar = 0;
+FAutoConsoleVariableRef ExtraAudioMixerDeviceLogging(
+	TEXT("au.ExtraAudioMixerDeviceLogging"),
+	ExtraAudioMixerDeviceLoggingCVar,
+	TEXT("Enables extra logging for audio mixer device running \n")
+	TEXT("0: no logging, 1: logging every 500 callbacks \n"),
+	ECVF_Default);
+
 namespace Audio
 {
 	int32 sRenderInstanceIds = 0;
@@ -224,6 +232,13 @@ namespace Audio
 			break;
 		}
 
+		static const int32 HeartBeatRate = 500;
+		if ((ExtraAudioMixerDeviceLoggingCVar > 0) && (++CallCounterMixNextBuffer > HeartBeatRate))
+		{
+			UE_LOG(LogAudioMixer, Display, TEXT("FOutputBuffer::MixNextBuffer() called %i times"), HeartBeatRate);
+			CallCounterMixNextBuffer = 0;
+		}
+
 		// Mark/signal that we're ready
 		bIsReady = true;
 		IsReadyEvent->Trigger();
@@ -396,6 +411,14 @@ namespace Audio
 	template<typename BufferType>
 	void IAudioMixerPlatformInterface::ApplyAttenuationInternal(BufferType* BufferDataPtr, const int32 NumFrames)
 	{
+		static const int32 HeartBeatRate = 500;
+		const bool bLog = (ExtraAudioMixerDeviceLoggingCVar > 0) && (++CallCounterApplyAttenuationInternal > HeartBeatRate);
+		if (bLog)
+		{
+			UE_LOG(LogAudioMixer, Display, TEXT("IAudioMixerPlatformInterface::ApplyAttenuationInternal() called %i times"), HeartBeatRate);
+			CallCounterApplyAttenuationInternal = 0;
+		}
+
 		// Perform fade in and fade out global attenuation to avoid clicks/pops on startup/shutdown
 		if (bPerformingFade)
 		{
@@ -409,11 +432,21 @@ namespace Audio
 			bFadedOut = (FadeVolume == 0.0f);
 			bPerformingFade = false;
 			AudioFadeEvent->Trigger();
+
+			if (bLog)
+			{
+				UE_LOG(LogAudioMixer, Display, TEXT("IAudioMixerPlatformInterface::ApplyAttenuationInternal() Faded from %f to %f"), FadeVolume, FadeParam.GetValue());
+			}
 		}
 		else if (bFadedOut)
 		{
 			// If we're faded out, then just zero the data.
 			FPlatformMemory::Memzero((void*)BufferDataPtr, sizeof(BufferType)*NumFrames);
+
+			if (bLog)
+			{
+				UE_LOG(LogAudioMixer, Display, TEXT("IAudioMixerPlatformInterface::ApplyAttenuationInternal() Zero'd out buffer"));
+			}
 		}
 
 		FadeParam.Reset();
@@ -421,6 +454,8 @@ namespace Audio
 
 	void IAudioMixerPlatformInterface::StartRunningNullDevice()
 	{
+		UE_LOG(LogAudioMixer, Display, TEXT("StartRunningNullDevice() called"));
+
 		if (!NullDeviceCallback.IsValid())
 		{
 			// Reset all of the buffers, then immediately kick off another render.
@@ -444,6 +479,8 @@ namespace Audio
 
 	void IAudioMixerPlatformInterface::StopRunningNullDevice()
 	{
+		UE_LOG(LogAudioMixer, Display, TEXT("StopRunningNullDevice() called"));
+
 		if (bIsUsingNullDevice)
 		{
 			CurrentBufferReadIndex = 0;
@@ -478,6 +515,13 @@ namespace Audio
 	{
 		LLM_SCOPE(ELLMTag::AudioMixer);
 
+		static const int32 HeartBeatRate = 500;
+		if ((ExtraAudioMixerDeviceLoggingCVar > 0) && (++CallCounterReadNextBuffer > HeartBeatRate))
+		{
+			UE_LOG(LogAudioMixer, Display, TEXT("IAudioMixerPlatformInterface::ReadNextBuffer() called %i times"), HeartBeatRate);
+			CallCounterReadNextBuffer = 0;
+		}
+
 		// If we are flushing buffers for our output voice and this is being called on the audio thread directly,
 		// early exit.
 		if (bIsInDeviceSwap)
@@ -585,6 +629,8 @@ namespace Audio
 
 		// Set the number of buffers to be one more than the number to queue.
 		NumOutputBuffers = FMath::Max(OpenStreamParams.NumBuffers, 2);
+		UE_LOG(LogAudioMixer, Display, TEXT("Output buffers initialized: Frames=%i, Channels=%i, Samples=%i"), NumOutputFrames, NumOutputChannels, NumOutputSamples);
+
 
 		OutputBuffers.Reset();
 		OutputBuffers.AddDefaulted(NumOutputBuffers);
@@ -686,6 +732,8 @@ namespace Audio
 
 	uint32 IAudioMixerPlatformInterface::RunInternal()
 	{
+		UE_LOG(LogAudioMixer, Display, TEXT("Starting AudioMixerPlatformInterface::RunInternal()"));
+
 		// Lets prime and submit the first buffer (which is going to be the buffer underrun buffer)
 		SubmitBuffer(UnderrunBuffer.GetBufferData());
 
diff --git a/Engine/Source/Runtime/AudioMixerCore/Public/AudioMixer.h b/Engine/Source/Runtime/AudioMixerCore/Public/AudioMixer.h
--- a/Engine/Source/Runtime/AudioMixerCore/Public/AudioMixer.h
+++ b/Engine/Source/Runtime/AudioMixerCore/Public/AudioMixer.h
@@ -337,7 +337,9 @@ namespace Audio
 		AlignedByteBuffer FormattedBuffer;
  		EAudioMixerStreamDataFormat::Type DataFormat;
  		FThreadSafeBool bIsReady;
- 	};
+
+		int32 CallCounterMixNextBuffer{ 0 };
+	};
 
 	/** Abstract interface for receiving audio device changed notifications */
 	class AUDIOMIXERCORE_API IAudioMixerDeviceChangedLister
@@ -598,6 +600,9 @@ namespace Audio
 		/** Struct used to store render time analysis data. */
 		FAudioRenderTimeAnalysis RenderTimeAnalysis;
 
+		int32 CallCounterApplyAttenuationInternal{ 0 };
+		int32 CallCounterReadNextBuffer{ 0 };
+
 		FThreadSafeBool bPerformingFade;
 		FThreadSafeBool bFadedOut;
 		FThreadSafeBool bIsDeviceInitialized;
diff --git a/Engine/Source/Runtime/Engine/Private/AudioDevice.cpp b/Engine/Source/Runtime/Engine/Private/AudioDevice.cpp
--- a/Engine/Source/Runtime/Engine/Private/AudioDevice.cpp
+++ b/Engine/Source/Runtime/Engine/Private/AudioDevice.cpp
@@ -498,7 +498,8 @@ bool FAudioDevice::Init(Audio::FDeviceId InDeviceID, int32 InMaxSources)
 	// allow the platform to startup
 	if (!InitializeHardware())
 	{
-		// Could not initialize hardware, tear down anything that was set up during initialization
+		// Could not initialize hardware. Tear down anything that was set up during initialization.
+		UE_LOG(LogAudio, Warning, TEXT("Could not initialize hardware. Tearing down anything that was set up during initialization"));
 		Teardown();
 
 		return false;
diff --git a/Engine/Source/Runtime/Windows/AudioMixerXAudio2/Private/AudioMixerPlatformXAudio2.cpp b/Engine/Source/Runtime/Windows/AudioMixerXAudio2/Private/AudioMixerPlatformXAudio2.cpp
--- a/Engine/Source/Runtime/Windows/AudioMixerXAudio2/Private/AudioMixerPlatformXAudio2.cpp
+++ b/Engine/Source/Runtime/Windows/AudioMixerXAudio2/Private/AudioMixerPlatformXAudio2.cpp
@@ -923,12 +923,17 @@ namespace Audio
 				DllName = GetNextDllToTry(DllName);
 				if (DllName != NAME_None)
 				{
+					UE_LOG(LogAudioMixer, Warning, TEXT("Xaudio 2.9: Failed to create Master Voice. Attempting fallback DLL"));
 					TeardownHardware();
 					if (InitializeHardware())
 					{
 						return OpenAudioStream(Params);
 					}
 				}
+				else
+				{
+					UE_LOG(LogAudioMixer, Warning, TEXT("Xaudio 2.9: Failed to create Master Voice. Out of DLL fallbacks"));
+				}
 			}
 #endif //PLATFORM_WINDOWS
 
@@ -982,6 +987,7 @@ namespace Audio
 
 	bool FMixerPlatformXAudio2::StartAudioStream()
 	{
+		UE_LOG(LogAudioMixer, Log, TEXT("FMixerPlatformXAudio2::StartAudioStream() called"));
 		// Start generating audio with our output source voice
 		BeginGeneratingAudio();
 
@@ -1010,6 +1016,8 @@ namespace Audio
 			return false;
 		}
 
+		UE_LOG(LogAudioMixer, Log, TEXT("FMixerPlatformXAudio2::StopAudioStream() called"));
+
 		check(XAudio2System);
 
 		if (AudioStreamInfo.StreamState != EAudioOutputStreamState::Stopped && AudioStreamInfo.StreamState != EAudioOutputStreamState::Closed)
@@ -1065,7 +1073,7 @@ namespace Audio
 			return true;
 		}
 
-		UE_LOG(LogTemp, Log, TEXT("Resetting audio stream to device id %s"), *InNewDeviceId);
+		UE_LOG(LogAudioMixer, Log, TEXT("Resetting audio stream to device id %s"), *InNewDeviceId);
 
 		if (bIsUsingNullDevice)
 		{
@@ -1226,6 +1234,12 @@ namespace Audio
 
 			// Submit buffer to the output streaming voice
 			OutputAudioStreamSourceVoice->SubmitSourceBuffer(&XAudio2Buffer);
+
+			if(!FirstBufferSubmitted)
+			{
+				UE_LOG(LogAudioMixer, Display, TEXT("FMixerPlatformXAudio2::SubmitBuffer() called for the first time"));
+				FirstBufferSubmitted = true;
+			}
 		}
 	}
 
diff --git a/Engine/Source/Runtime/Windows/AudioMixerXAudio2/Private/AudioMixerPlatformXAudio2.h b/Engine/Source/Runtime/Windows/AudioMixerXAudio2/Private/AudioMixerPlatformXAudio2.h
--- a/Engine/Source/Runtime/Windows/AudioMixerXAudio2/Private/AudioMixerPlatformXAudio2.h
+++ b/Engine/Source/Runtime/Windows/AudioMixerXAudio2/Private/AudioMixerPlatformXAudio2.h
@@ -123,6 +123,8 @@ namespace Audio
 		// we check whether a new audio device was connected every second or so.
 		float TimeSinceNullDeviceWasLastChecked;
 
+		bool FirstBufferSubmitted{false};
+
 		uint32 bIsInitialized : 1;
 		uint32 bIsDeviceOpen : 1;
 
