Clean-up using TryLock
#rnx

diff --git a/Engine/Source/Runtime/AudioMixerCore/Private/AudioMixer.cpp b/Engine/Source/Runtime/AudioMixerCore/Private/AudioMixer.cpp
--- a/Engine/Source/Runtime/AudioMixerCore/Private/AudioMixer.cpp
+++ b/Engine/Source/Runtime/AudioMixerCore/Private/AudioMixer.cpp
@@ -3,11 +3,12 @@
 #include "AudioMixer.h"
 #include "DSP/BufferVectorOperations.h"
 #include "HAL/RunnableThread.h"
-#include "Misc/ConfigCacheIni.h"
 #include "HAL/ThreadSafeCounter.h"
 #include "HAL/IConsoleManager.h"
 #include "HAL/Event.h"
 #include "HAL/LowLevelMemTracker.h"
+#include "Misc/ConfigCacheIni.h"
+#include "Misc/ScopeTryLock.h"
 #include "ProfilingDebugging/CsvProfiler.h"
 
 // Defines the "Audio" category in the CSV profiler.
@@ -529,85 +530,85 @@ namespace Audio
 			return;
 		}
 
-		// If we are currently swapping devices and OnBufferEnd is being triggered in an XAudio2Thread,
-		// early exit.
-		if (!DeviceSwapCriticalSection.TryLock())
 		{
-			return;
-		}
+			FScopeTryLock DeviceSwapLock(&DeviceSwapCriticalSection);
 
-		// Don't read any more audio if we're not running or changing device
-		if (AudioStreamInfo.StreamState != EAudioOutputStreamState::Running)
-		{
-			DeviceSwapCriticalSection.Unlock();
-			return;
-		}
+			// If we are currently swapping devices and OnBufferEnd is being triggered in an XAudio2Thread,
+			// early exit.
+			if (!DeviceSwapLock.IsLocked())
+			{
+				return;
+			}
 
-		// AudioRenderThread hasn't executed yet, return silence
-		if (CurrentBufferReadIndex == INDEX_NONE || CurrentBufferWriteIndex == INDEX_NONE)
-		{
-			SubmitBuffer(UnderrunBuffer.GetBufferData());
-			DeviceSwapCriticalSection.Unlock();
-			return;
-		}
+			// Don't read any more audio if we're not running or changing device
+			if (AudioStreamInfo.StreamState != EAudioOutputStreamState::Running)
+			{
+				return;
+			}
 
-		// Reset the ready state of the buffer which was just finished playing
-		FOutputBuffer& CurrentReadBuffer = OutputBuffers[CurrentBufferReadIndex];
-		CurrentReadBuffer.ResetReadyState();
+			// AudioRenderThread hasn't executed yet, return silence
+			if (CurrentBufferReadIndex == INDEX_NONE || CurrentBufferWriteIndex == INDEX_NONE)
+			{
+				SubmitBuffer(UnderrunBuffer.GetBufferData());
+				return;
+			}
 
-		// Get the next index that we want to read
-		int32 NextReadIndex = (CurrentBufferReadIndex + 1) % NumOutputBuffers;
+			// Reset the ready state of the buffer which was just finished playing
+			FOutputBuffer& CurrentReadBuffer = OutputBuffers[CurrentBufferReadIndex];
+			CurrentReadBuffer.ResetReadyState();
 
-		// If it's not ready, warn, and then wait here. This will cause underruns but is preferable than getting out-of-order buffer state.
-		static int32 UnderrunCount = 0;
-		static int32 CurrentUnderrunCount = 0;
+			// Get the next index that we want to read
+			int32 NextReadIndex = (CurrentBufferReadIndex + 1) % NumOutputBuffers;
 
-		bool bSubmittingUnderrunBuffer = false;
+			// If it's not ready, warn, and then wait here. This will cause underruns but is preferable than getting out-of-order buffer state.
+			static int32 UnderrunCount = 0;
+			static int32 CurrentUnderrunCount = 0;
 
-		if (!OutputBuffers[NextReadIndex].IsReady())
-		{
-			// try to wait for the buffer to be ready
-			FEvent* BufferReadyEvent = OutputBuffers[NextReadIndex].IsReadyEvent;
-			if (!BufferReadyEvent || !BufferReadyEvent->Wait(static_cast<uint32>(UnderrunTimeoutCVar)))
+			bool bSubmittingUnderrunBuffer = false;
+
+			if (!OutputBuffers[NextReadIndex].IsReady())
 			{
-				bSubmittingUnderrunBuffer = true; // Event didn't fire in time
-			}
-		}
-		
-		if (bSubmittingUnderrunBuffer)
-		{
-			UnderrunCount++;
-			CurrentUnderrunCount++;
-			
-			if (!bWarnedBufferUnderrun)
-			{						
-				UE_LOG(LogAudioMixer, Display, TEXT("Audio Buffer Underrun detected."));
-				bWarnedBufferUnderrun = true;
+				// try to wait for the buffer to be ready
+				FEvent* BufferReadyEvent = OutputBuffers[NextReadIndex].IsReadyEvent;
+				if (!BufferReadyEvent || !BufferReadyEvent->Wait(static_cast<uint32>(UnderrunTimeoutCVar)))
+				{
+					bSubmittingUnderrunBuffer = true; // Event didn't fire in time
+				}
 			}
-		
-			SubmitBuffer(UnderrunBuffer.GetBufferData());
-		}
-		else
-		{
-			ApplyMasterAttenuation();
 
-			// As soon as a valid buffer goes through, allow more warning
-			if (bWarnedBufferUnderrun)
+			if (bSubmittingUnderrunBuffer)
 			{
-				UE_LOG(LogAudioMixerDebug, Log, TEXT("Audio had %d underruns [Total: %d]."), CurrentUnderrunCount, UnderrunCount);
+				UnderrunCount++;
+				CurrentUnderrunCount++;
+
+				if (!bWarnedBufferUnderrun)
+				{
+					UE_LOG(LogAudioMixer, Display, TEXT("Audio Buffer Underrun detected."));
+					bWarnedBufferUnderrun = true;
+				}
+
+				SubmitBuffer(UnderrunBuffer.GetBufferData());
 			}
-			CurrentUnderrunCount = 0;
-			bWarnedBufferUnderrun = false;
+			else
+			{
+				ApplyMasterAttenuation();
 
-			// Submit the buffer at the next read index, but don't set the read index value yet
-			SubmitBuffer(OutputBuffers[NextReadIndex].GetBufferData());
+				// As soon as a valid buffer goes through, allow more warning
+				if (bWarnedBufferUnderrun)
+				{
+					UE_LOG(LogAudioMixerDebug, Log, TEXT("Audio had %d underruns [Total: %d]."), CurrentUnderrunCount, UnderrunCount);
+				}
+				CurrentUnderrunCount = 0;
+				bWarnedBufferUnderrun = false;
 
-			// Update the current read index to the next read index
-			CurrentBufferReadIndex = NextReadIndex;
-			OutputBuffers[NextReadIndex].IsReadyEvent->Reset();
-		}
+				// Submit the buffer at the next read index, but don't set the read index value yet
+				SubmitBuffer(OutputBuffers[NextReadIndex].GetBufferData());
 
-		DeviceSwapCriticalSection.Unlock();
+				// Update the current read index to the next read index
+				CurrentBufferReadIndex = NextReadIndex;
+				OutputBuffers[NextReadIndex].IsReadyEvent->Reset();
+			}
+		}
 
 		// Kick off rendering of the next set of buffers
 		if (AudioRenderEvent)
