Fix for shutdown for single-threaded applications. FadeOut waits for an FEvent that never gets hit.

diff --git a/Engine/Source/Runtime/AudioMixerCore/Private/AudioMixer.cpp b/Engine/Source/Runtime/AudioMixerCore/Private/AudioMixer.cpp
--- a/Engine/Source/Runtime/AudioMixerCore/Private/AudioMixer.cpp
+++ b/Engine/Source/Runtime/AudioMixerCore/Private/AudioMixer.cpp
@@ -353,7 +353,7 @@ namespace Audio
 	{
 		// Non Realtime isn't ticked when fade out is called, and the user can't hear
 		// the output anyways so there's no need to make it pleasant for their ears.
-		if (IsNonRealtime())
+		if (!FPlatformProcess::SupportsMultithreading() || IsNonRealtime())
 		{
 			bFadedOut = true;
 			FadeVolume = 0.f;
