Circular buffer refactor of AudioMixer produce/consume cycles.  (code by Ethan).

diff --git a/Engine/Plugins/Lumin/MagicLeap/Source/MagicLeapAudio/Private/AudioMixerPlatformMagicLeap.cpp b/Engine/Plugins/Lumin/MagicLeap/Source/MagicLeapAudio/Private/AudioMixerPlatformMagicLeap.cpp
--- a/Engine/Plugins/Lumin/MagicLeap/Source/MagicLeapAudio/Private/AudioMixerPlatformMagicLeap.cpp
+++ b/Engine/Plugins/Lumin/MagicLeap/Source/MagicLeapAudio/Private/AudioMixerPlatformMagicLeap.cpp
@@ -574,15 +574,6 @@ namespace Audio
 			// StopRunningNullDevice blocks until the sleeping thread wakes up
 			StopRunningNullDevice();
 
-			// Reset UE buffers.
-			for (int32 Index = 0; Index < OutputBuffers.Num(); ++Index)
-			{
-				OutputBuffers[Index].Reset(OpenStreamParams.NumFrames * AudioStreamInfo.DeviceInfo.NumChannels);
-			}
-
-			CurrentBufferReadIndex = 0;
-			CurrentBufferWriteIndex = 1;
-
 #if WITH_MLSDK
 			// Reset the device buffers
 			// Stopping and starting the audio causes the streaming buffer position to be set to 0
diff --git a/Engine/Source/Runtime/AudioMixerCore/Private/AudioMixer.cpp b/Engine/Source/Runtime/AudioMixerCore/Private/AudioMixer.cpp
--- a/Engine/Source/Runtime/AudioMixerCore/Private/AudioMixer.cpp
+++ b/Engine/Source/Runtime/AudioMixerCore/Private/AudioMixer.cpp
@@ -147,56 +147,43 @@ namespace Audio
 		}
 	}
 
-	FOutputBuffer::~FOutputBuffer()
-	{
-		if (IsReadyEvent != nullptr)
-		{
-			FPlatformProcess::ReturnSynchEventToPool(IsReadyEvent);
-			IsReadyEvent = nullptr;
-		}
-	}
 
-
-	void FOutputBuffer::Init(IAudioMixer* InAudioMixer, const int32 InNumSamples, const EAudioMixerStreamDataFormat::Type InDataFormat)
+	void FOutputBuffer::Init(IAudioMixer* InAudioMixer, const int32 InNumSamples, const int32 InNumBuffers, const EAudioMixerStreamDataFormat::Type InDataFormat)
 	{
-		Buffer.SetNumZeroed(InNumSamples);
+		RenderBuffer.Reset();
+		RenderBuffer.AddUninitialized(InNumSamples);
+
 		DataFormat = InDataFormat;
 
 		check(InAudioMixer != nullptr);
 		AudioMixer = InAudioMixer;
 
-		if (IsReadyEvent == nullptr)
-		{
-			IsReadyEvent = FPlatformProcess::GetSynchEventFromPool(true /*Manual Reset*/);
-		}
-		check(IsReadyEvent != nullptr);
-
-		switch (DataFormat)
-		{
-			case EAudioMixerStreamDataFormat::Float:
-				// nothing to do...
-				break;
+		CircularBuffer.SetCapacity(InNumSamples * InNumBuffers * GetSizeForDataFormat(DataFormat));
 
-			case EAudioMixerStreamDataFormat::Int16:
-				FormattedBuffer.SetNumZeroed(InNumSamples * sizeof(int16));	
-				break;
+		PopBuffer.Reset();
+		PopBuffer.AddUninitialized(InNumSamples * GetSizeForDataFormat(DataFormat));
 
-			default:
-				// Not implemented/supported
-				check(false);
-				break;
+		if (DataFormat != EAudioMixerStreamDataFormat::Float)
+		{
+			FormattedBuffer.SetNumZeroed(InNumSamples * GetSizeForDataFormat(DataFormat));
 		}
 	}
 
-	void FOutputBuffer::MixNextBuffer()
+	bool FOutputBuffer::MixNextBuffer()
  	{
+		// If the circular queue is already full, exit.
+		if (CircularBuffer.Remainder() < static_cast<uint32>(RenderBuffer.Num()))
+		{
+			return false;
+		}
+
 		CSV_SCOPED_TIMING_STAT(Audio, RenderAudio);
 
 		// Zero the buffer
-		FPlatformMemory::Memzero(Buffer.GetData(), Buffer.Num() * sizeof(float));
+		FPlatformMemory::Memzero(RenderBuffer.GetData(), RenderBuffer.Num() * sizeof(float));
 		if (AudioMixer != nullptr)
 		{
-			AudioMixer->OnProcessAudioStream(Buffer);
+			AudioMixer->OnProcessAudioStream(RenderBuffer);
 		}
 
 		switch (DataFormat)
@@ -205,25 +192,31 @@ namespace Audio
 		{
 			if (!FMath::IsNearlyEqual(LinearGainScalarForFinalOututCVar, 1.0f))
 			{
-				MultiplyBufferByConstantInPlace(Buffer, LinearGainScalarForFinalOututCVar);
+				MultiplyBufferByConstantInPlace(RenderBuffer, LinearGainScalarForFinalOututCVar);
 			}
-			BufferRangeClampFast(Buffer, -1.0f, 1.0f);
+			BufferRangeClampFast(RenderBuffer, -1.0f, 1.0f);
+
+			// No conversion is needed, so we push the RenderBuffer directly to the circular queue.
+			CircularBuffer.Push(reinterpret_cast<const uint8*>(RenderBuffer.GetData()), RenderBuffer.Num() * sizeof(float));
 		}
 		break;
 
 		case EAudioMixerStreamDataFormat::Int16:
 		{
 			int16* BufferInt16 = (int16*)FormattedBuffer.GetData();
-			const int32 NumSamples = Buffer.Num();
+			const int32 NumSamples = RenderBuffer.Num();
+			check(FormattedBuffer.Num() / GetSizeForDataFormat(DataFormat) == RenderBuffer.Num());			
 
 			const float ConversionScalar = LinearGainScalarForFinalOututCVar * 32767.0f;
-			MultiplyBufferByConstantInPlace(Buffer, ConversionScalar);
-			BufferRangeClampFast(Buffer, -32767.0f, 32767.0f);
+			MultiplyBufferByConstantInPlace(RenderBuffer, ConversionScalar);
+			BufferRangeClampFast(RenderBuffer, -32767.0f, 32767.0f);
 
 			for (int32 i = 0; i < NumSamples; ++i)
 			{
-				BufferInt16[i] = (int16)Buffer[i];
+				BufferInt16[i] = (int16)RenderBuffer[i];
 			}
+
+			CircularBuffer.Push(reinterpret_cast<const uint8*>(FormattedBuffer.GetData()), FormattedBuffer.Num());
 		}
 		break;
 
@@ -240,70 +233,36 @@ namespace Audio
 			CallCounterMixNextBuffer = 0;
 		}
 
-		// Mark/signal that we're ready
-		bIsReady = true;
-		IsReadyEvent->Trigger();
+		return true;
  	}
  
-	const uint8* FOutputBuffer::GetBufferData() const
+	TArrayView<const uint8> FOutputBuffer::PopBufferData(int32& OutNumBytesPopped) const
 	{
-		if (DataFormat == EAudioMixerStreamDataFormat::Float)
-		{
-			return (const uint8*)Buffer.GetData();
-		}
-		else
-		{
-			return (const uint8*)FormattedBuffer.GetData();
-		}
-	}
+		FMemory::Memzero(reinterpret_cast<uint8*>(PopBuffer.GetData()), PopBuffer.Num());
+		OutNumBytesPopped = CircularBuffer.Pop(PopBuffer.GetData(), PopBuffer.Num());
 
-	uint8* FOutputBuffer::GetBufferData()
-	{
-		if (DataFormat == EAudioMixerStreamDataFormat::Float)
-		{
-			return (uint8*)Buffer.GetData();
-		}
-		else
-		{
-			return (uint8*)FormattedBuffer.GetData();
-		}
+		return TArrayView<const uint8>(PopBuffer);
 	}
 
-	int32 FOutputBuffer::GetNumFrames() const
+	int32 FOutputBuffer::GetNumSamples() const
 	{
-		return Buffer.Num();
+		return RenderBuffer.Num();
 	}
 
-	void FOutputBuffer::ResetReadyState()
+	size_t FOutputBuffer::GetSizeForDataFormat(EAudioMixerStreamDataFormat::Type InDataFormat)
 	{
-		bIsReady = false;
-		if (IsReadyEvent)
+		switch (InDataFormat)
 		{
-			IsReadyEvent->Reset();
-		}
-	}
-
-
-	void FOutputBuffer::Reset(const int32 InNewNumSamples)
-	{
-		Buffer.Reset();
-		Buffer.AddZeroed(InNewNumSamples);
+		case EAudioMixerStreamDataFormat::Float:
+			return sizeof(float);
 
-		switch (DataFormat)
-		{
-			// Doesn't do anything...
-			case EAudioMixerStreamDataFormat::Float:
-				break;
+		case EAudioMixerStreamDataFormat::Int16:
+			return sizeof(int16);
 
-			case EAudioMixerStreamDataFormat::Int16:
-			{
-				FormattedBuffer.Reset();
-				FormattedBuffer.AddZeroed(InNewNumSamples * sizeof(int16));
-			}
-			break;
+		default:
+			checkNoEntry();
+			return 0;
 		}
-
-		bIsReady = false;
 	}
 
 	/**
@@ -316,8 +275,6 @@ namespace Audio
 		, AudioRenderEvent(nullptr)
 		, bIsInDeviceSwap(false)
 		, AudioFadeEvent(nullptr)
-		, CurrentBufferReadIndex(INDEX_NONE)
-		, CurrentBufferWriteIndex(INDEX_NONE)
 		, NumOutputBuffers(0)
 		, FadeVolume(0.0f)
 		, LastError(TEXT("None"))
@@ -410,7 +367,7 @@ namespace Audio
 	}
 
 	template<typename BufferType>
-	void IAudioMixerPlatformInterface::ApplyAttenuationInternal(BufferType* BufferDataPtr, const int32 NumFrames)
+	void IAudioMixerPlatformInterface::ApplyAttenuationInternal(TArrayView<BufferType>& InOutBuffer)
 	{
 		static const int32 HeartBeatRate = 500;
 		const bool bLog = (ExtraAudioMixerDeviceLoggingCVar > 0) && (++CallCounterApplyAttenuationInternal > HeartBeatRate);
@@ -423,11 +380,11 @@ namespace Audio
 		// Perform fade in and fade out global attenuation to avoid clicks/pops on startup/shutdown
 		if (bPerformingFade)
 		{
-			FadeParam.SetValue(FadeVolume, NumFrames);
+			FadeParam.SetValue(FadeVolume, InOutBuffer.Num());
 
-			for (int32 i = 0; i < NumFrames; ++i)
+			for (int32 i = 0; i < InOutBuffer.Num(); ++i)
 			{
-				BufferDataPtr[i] = (BufferType)(BufferDataPtr[i] * FadeParam.Update());
+				InOutBuffer[i] = (BufferType)(InOutBuffer[i] * FadeParam.Update());
 			}
 
 			bFadedOut = (FadeVolume == 0.0f);
@@ -442,7 +399,7 @@ namespace Audio
 		else if (bFadedOut)
 		{
 			// If we're faded out, then just zero the data.
-			FPlatformMemory::Memzero((void*)BufferDataPtr, sizeof(BufferType)*NumFrames);
+			FPlatformMemory::Memzero((void*)InOutBuffer.GetData(), sizeof(BufferType)* InOutBuffer.Num());
 
 			if (bLog)
 			{
@@ -459,13 +416,8 @@ namespace Audio
 
 		if (!NullDeviceCallback.IsValid())
 		{
-			// Reset all of the buffers, then immediately kick off another render.
-			for (int32 Index = 0; Index < OutputBuffers.Num(); ++Index)
-			{
-				OutputBuffers[Index].Reset(OpenStreamParams.NumFrames * AudioStreamInfo.DeviceInfo.NumChannels);
-			}
 
-			check(OpenStreamParams.NumFrames * AudioStreamInfo.DeviceInfo.NumChannels == OutputBuffers[0].GetBuffer().Num());
+			check(OpenStreamParams.NumFrames * AudioStreamInfo.DeviceInfo.NumChannels == OutputBuffer.GetNumSamples());
 
 			AudioRenderEvent->Trigger();
 
@@ -482,11 +434,6 @@ namespace Audio
 	{
 		UE_LOG(LogAudioMixer, Display, TEXT("StopRunningNullDevice() called"));
 
-		if (bIsUsingNullDevice)
-		{
-			CurrentBufferReadIndex = 0;
-			CurrentBufferWriteIndex = 1;
-		}
 		if (NullDeviceCallback.IsValid())
 		{
 			NullDeviceCallback.Reset();
@@ -494,21 +441,23 @@ namespace Audio
 		}
 	}
 
-	void IAudioMixerPlatformInterface::ApplyMasterAttenuation()
+	void IAudioMixerPlatformInterface::ApplyMasterAttenuation(TArrayView<const uint8>& OutPoppedAudio)
 	{
-		const int32 NextReadIndex = (CurrentBufferReadIndex + 1) % NumOutputBuffers;
-		FOutputBuffer& CurrentReadBuffer = OutputBuffers[NextReadIndex];
+		EAudioMixerStreamDataFormat::Type Format = OutputBuffer.GetFormat();
 
-		EAudioMixerStreamDataFormat::Type Format = CurrentReadBuffer.GetFormat();
-		uint8* BufferDataPtr = CurrentReadBuffer.GetBufferData();
-		
 		if (Format == EAudioMixerStreamDataFormat::Float)
 		{
-			ApplyAttenuationInternal((float*)BufferDataPtr, CurrentReadBuffer.GetNumFrames());
+			TArrayView<float> OutFloatBuffer = TArrayView<float>(const_cast<float*>(reinterpret_cast<const float*>(OutPoppedAudio.GetData())), OutPoppedAudio.Num() / sizeof(float));
+			ApplyAttenuationInternal(OutFloatBuffer);
+		}
+		else if (Format == EAudioMixerStreamDataFormat::Int16)
+		{
+			TArrayView<int16> OutIntBuffer = TArrayView<int16>(const_cast<int16*>(reinterpret_cast<const int16*>(OutPoppedAudio.GetData())), OutPoppedAudio.Num() / sizeof(int16));
+			ApplyAttenuationInternal(OutIntBuffer);
 		}
 		else
 		{
-			ApplyAttenuationInternal((int16*)BufferDataPtr, CurrentReadBuffer.GetNumFrames());
+			checkNoEntry();
 		}
 	}
 
@@ -516,13 +465,6 @@ namespace Audio
 	{
 		LLM_SCOPE(ELLMTag::AudioMixer);
 
-		static const int32 HeartBeatRate = 500;
-		if ((ExtraAudioMixerDeviceLoggingCVar > 0) && (++CallCounterReadNextBuffer > HeartBeatRate))
-		{
-			UE_LOG(LogAudioMixer, Display, TEXT("IAudioMixerPlatformInterface::ReadNextBuffer() called %i times"), HeartBeatRate);
-			CallCounterReadNextBuffer = 0;
-		}
-
 		// If we are flushing buffers for our output voice and this is being called on the audio thread directly,
 		// early exit.
 		if (bIsInDeviceSwap)
@@ -530,85 +472,56 @@ namespace Audio
 			return;
 		}
 
+		// If we are currently swapping devices and OnBufferEnd is being triggered in an XAudio2Thread,
+		// early exit.
+		if (!DeviceSwapCriticalSection.TryLock())
 		{
-			FScopeTryLock DeviceSwapLock(&DeviceSwapCriticalSection);
-
-			// If we are currently swapping devices and OnBufferEnd is being triggered in an XAudio2Thread,
-			// early exit.
-			if (!DeviceSwapLock.IsLocked())
-			{
-				return;
-			}
-
-			// Don't read any more audio if we're not running or changing device
-			if (AudioStreamInfo.StreamState != EAudioOutputStreamState::Running)
-			{
-				return;
-			}
-
-			// AudioRenderThread hasn't executed yet, return silence
-			if (CurrentBufferReadIndex == INDEX_NONE || CurrentBufferWriteIndex == INDEX_NONE)
-			{
-				SubmitBuffer(UnderrunBuffer.GetBufferData());
-				return;
-			}
-
-			// Reset the ready state of the buffer which was just finished playing
-			FOutputBuffer& CurrentReadBuffer = OutputBuffers[CurrentBufferReadIndex];
-			CurrentReadBuffer.ResetReadyState();
+			return;
+		}
 
-			// Get the next index that we want to read
-			int32 NextReadIndex = (CurrentBufferReadIndex + 1) % NumOutputBuffers;
+		// Don't read any more audio if we're not running or changing device
+		if (AudioStreamInfo.StreamState != EAudioOutputStreamState::Running)
+		{
+			DeviceSwapCriticalSection.Unlock();
+			return;
+		}
 
-			// If it's not ready, warn, and then wait here. This will cause underruns but is preferable than getting out-of-order buffer state.
-			static int32 UnderrunCount = 0;
-			static int32 CurrentUnderrunCount = 0;
+		// If it's not ready, warn, and then wait here. This will cause underruns but is preferable than getting out-of-order buffer state.
+		static int32 UnderrunCount = 0;
+		static int32 CurrentUnderrunCount = 0;
 
-			bool bSubmittingUnderrunBuffer = false;
+		int32 NumSamplesPopped = 0;
+		TArrayView<const uint8> PoppedAudio = OutputBuffer.PopBufferData(NumSamplesPopped);
 
-			if (!OutputBuffers[NextReadIndex].IsReady())
-			{
-				// try to wait for the buffer to be ready
-				FEvent* BufferReadyEvent = OutputBuffers[NextReadIndex].IsReadyEvent;
-				if (!BufferReadyEvent || !BufferReadyEvent->Wait(static_cast<uint32>(UnderrunTimeoutCVar)))
-				{
-					bSubmittingUnderrunBuffer = true; // Event didn't fire in time
-				}
+		bool bDidOutputUnderrun = NumSamplesPopped != PoppedAudio.Num();
+		
+		if (bDidOutputUnderrun)
+		{
+			UnderrunCount++;
+			CurrentUnderrunCount++;
+			
+			if (!bWarnedBufferUnderrun)
+			{						
+				UE_LOG(LogAudioMixer, Display, TEXT("Audio Buffer Underrun detected. Things to try: Increase # output buffers, ensure audio-render thread has time to run (affinity and priority), debug your mix and reduce # sounds playing."));
+				bWarnedBufferUnderrun = true;
 			}
-
-			if (bSubmittingUnderrunBuffer)
+		}
+		else
+		{
+			// As soon as a valid buffer goes through, allow more warning
+				if (bWarnedBufferUnderrun)
 			{
-				UnderrunCount++;
-				CurrentUnderrunCount++;
-
-				if (!bWarnedBufferUnderrun)
-				{
-					UE_LOG(LogAudioMixer, Display, TEXT("Audio Buffer Underrun detected."));
-					bWarnedBufferUnderrun = true;
-				}
-
-				SubmitBuffer(UnderrunBuffer.GetBufferData());
+				UE_LOG(LogAudioMixerDebug, Log, TEXT("Audio had %d underruns [Total: %d]."), CurrentUnderrunCount, UnderrunCount);
 			}
-			else
-			{
-				ApplyMasterAttenuation();
 
-				// As soon as a valid buffer goes through, allow more warning
-				if (bWarnedBufferUnderrun)
-				{
-					UE_LOG(LogAudioMixerDebug, Log, TEXT("Audio had %d underruns [Total: %d]."), CurrentUnderrunCount, UnderrunCount);
-				}
-				CurrentUnderrunCount = 0;
-				bWarnedBufferUnderrun = false;
+			CurrentUnderrunCount = 0;
+			bWarnedBufferUnderrun = false;
+		}
 
-				// Submit the buffer at the next read index, but don't set the read index value yet
-				SubmitBuffer(OutputBuffers[NextReadIndex].GetBufferData());
+		ApplyMasterAttenuation(PoppedAudio);
+		SubmitBuffer(PoppedAudio.GetData());
 
-				// Update the current read index to the next read index
-				CurrentBufferReadIndex = NextReadIndex;
-				OutputBuffers[NextReadIndex].IsReadyEvent->Reset();
-			}
-		}
+		DeviceSwapCriticalSection.Unlock();
 
 		// Kick off rendering of the next set of buffers
 		if (AudioRenderEvent)
@@ -633,15 +546,7 @@ namespace Audio
 		UE_LOG(LogAudioMixer, Display, TEXT("Output buffers initialized: Frames=%i, Channels=%i, Samples=%i"), NumOutputFrames, NumOutputChannels, NumOutputSamples);
 
 
-		OutputBuffers.Reset();
-		OutputBuffers.AddDefaulted(NumOutputBuffers);
-		for (int32 Index = 0; Index < NumOutputBuffers; ++Index)
-		{
-			OutputBuffers[Index].Init(AudioStreamInfo.AudioMixer, NumOutputSamples, AudioStreamInfo.DeviceInfo.Format);
-		}
-
-		// Create an underrun buffer
-		UnderrunBuffer.Init(AudioStreamInfo.AudioMixer, NumOutputSamples, AudioStreamInfo.DeviceInfo.Format);
+		OutputBuffer.Init(AudioStreamInfo.AudioMixer, NumOutputSamples, NumOutputBuffers, AudioStreamInfo.DeviceInfo.Format);
 
 		AudioStreamInfo.StreamState = EAudioOutputStreamState::Running;
 
@@ -715,13 +620,8 @@ namespace Audio
 		if (AudioStreamInfo.StreamState == EAudioOutputStreamState::Running && bIsDeviceInitialized)
 		{
 			// Render mixed buffers till our queued buffers are filled up
-			while (CurrentBufferReadIndex != CurrentBufferWriteIndex)
+			while (OutputBuffer.MixNextBuffer())
 			{
-				RenderTimeAnalysis.Start();
-				OutputBuffers[CurrentBufferWriteIndex].MixNextBuffer();
-				RenderTimeAnalysis.End();
-
-				CurrentBufferWriteIndex = (CurrentBufferWriteIndex + 1) % NumOutputBuffers;
 			}
 		}
 	}
@@ -736,31 +636,20 @@ namespace Audio
 		UE_LOG(LogAudioMixer, Display, TEXT("Starting AudioMixerPlatformInterface::RunInternal()"));
 
 		// Lets prime and submit the first buffer (which is going to be the buffer underrun buffer)
-		SubmitBuffer(UnderrunBuffer.GetBufferData());
-
-		OutputBuffers[0].MixNextBuffer();
+		int32 NumSamplesPopped;
+		TArrayView<const uint8> AudioToSubmit = OutputBuffer.PopBufferData(NumSamplesPopped);
 
-		// Start immediately processing the next buffer
-		checkf(CurrentBufferReadIndex == INDEX_NONE, TEXT("CurrentBufferReadIndex: %i, StreamState: %i"), CurrentBufferReadIndex.Load(), AudioStreamInfo.StreamState);
-		checkf(CurrentBufferWriteIndex == INDEX_NONE, TEXT("CurrentBufferWriteIndex: %i, StreamState: %i"), CurrentBufferWriteIndex.Load(), AudioStreamInfo.StreamState);
+		SubmitBuffer(AudioToSubmit.GetData());
 
-		CurrentBufferReadIndex = 0;
-		CurrentBufferWriteIndex = 1;
+		OutputBuffer.MixNextBuffer();
 
 		while (AudioStreamInfo.StreamState != EAudioOutputStreamState::Stopping)
 		{
-			RenderTimeAnalysis.Start();
-
 			// Render mixed buffers till our queued buffers are filled up
-			while (CurrentBufferReadIndex != CurrentBufferWriteIndex && bIsDeviceInitialized)
+			while (bIsDeviceInitialized && OutputBuffer.MixNextBuffer())
 			{
-				OutputBuffers[CurrentBufferWriteIndex].MixNextBuffer();
-
-				CurrentBufferWriteIndex = (CurrentBufferWriteIndex + 1) % NumOutputBuffers;
 			}
 
-			RenderTimeAnalysis.End();
-
 			// Bounds check the timeout for our audio render event.
 			OverrunTimeoutCVar = FMath::Clamp(OverrunTimeoutCVar, 500, 5000);
 
@@ -773,8 +662,6 @@ namespace Audio
 			}
 		}
 
-		CurrentBufferReadIndex = INDEX_NONE;
-		CurrentBufferWriteIndex = INDEX_NONE;
 		OpenStreamParams.AudioMixer->OnAudioStreamShutdown();
 
 		AudioStreamInfo.StreamState = EAudioOutputStreamState::Stopped;
diff --git a/Engine/Source/Runtime/AudioMixerCore/Public/AudioMixer.h b/Engine/Source/Runtime/AudioMixerCore/Public/AudioMixer.h
--- a/Engine/Source/Runtime/AudioMixerCore/Public/AudioMixer.h
+++ b/Engine/Source/Runtime/AudioMixerCore/Public/AudioMixer.h
@@ -12,6 +12,7 @@
 #include "AudioMixerNullDevice.h"
 #include "DSP/ParamInterpolator.h"
 #include "DSP/BufferVectorOperations.h"
+#include "DSP/Dsp.h"
 #include "Modules/ModuleInterface.h"
 
 // defines used for AudioMixer.h
@@ -291,53 +292,45 @@ namespace Audio
 	{
 	public:
 		FOutputBuffer()
-			: IsReadyEvent(nullptr)
-			, AudioMixer(nullptr)
+			: AudioMixer(nullptr)
 			, DataFormat(EAudioMixerStreamDataFormat::Unknown)
 		{}
 
-		~FOutputBuffer();
+		~FOutputBuffer() = default;
  
 		/** Initialize the buffer with the given samples and output format. */
-		void Init(IAudioMixer* InAudioMixer, const int32 InNumSamples, const EAudioMixerStreamDataFormat::Type InDataFormat);
-
-		/** Gets the next mixed buffer from the audio mixer. */
-		void MixNextBuffer();
-
-		/** Returns the float buffer. */
-		AlignedFloatBuffer& GetBuffer() { return Buffer; }
-
+		void Init(IAudioMixer* InAudioMixer, const int32 InNumSamples, const int32 InNumBuffers, const EAudioMixerStreamDataFormat::Type InDataFormat);
 
-		/** Gets the buffer data ptrs. */
-		const uint8* GetBufferData() const;
-		uint8* GetBufferData();
+		/** Gets the next mixed buffer from the audio mixer. Returns false if our buffer is already full. */
+		bool MixNextBuffer();
 
-		/** Event to signal that the buffer is ready */
-		FEvent* IsReadyEvent;
+		/** Gets the buffer data ptrs. Returns a TArrayView for the full buffer size requested, but in the case of an underrun, OutBytesPopped will be less that the size of the returned TArrayView. */
+		TArrayView<const uint8> PopBufferData(int32& OutBytesPopped) const;
 
 		/** Gets the number of frames of the buffer. */
-		int32 GetNumFrames() const;
+		int32 GetNumSamples() const;
 
 		/** Returns the format of the buffer. */
 		EAudioMixerStreamDataFormat::Type GetFormat() const { return DataFormat; }
-		
-		/** Returns if ready. */
-		bool IsReady() const { return bIsReady; }
-
-		/** Resets the buffer ready state. */
-		void ResetReadyState();
-
-		/** Resets the internal buffers to the new sample count. Used when device is changed. */
-		void Reset(const int32 InNewNumSamples);
 
 
 	private:
 		IAudioMixer* AudioMixer;
-		AlignedFloatBuffer Buffer;
+
+		// Circular buffer used to buffer audio between the audio render thread and the platform interface thread.
+		mutable Audio::TCircularAudioBuffer<uint8> CircularBuffer;
+		
+		// Buffer that we render audio to from the IAudioMixer instance associated with this output buffer.
+		Audio::AlignedFloatBuffer RenderBuffer;
+
+		// Buffer read by the platform interface thread.
+		mutable Audio::AlignedByteBuffer PopBuffer;
+
+		// For non-float situations, this buffer is used to convert RenderBuffer before pushing it to CircularBuffer.
 		AlignedByteBuffer FormattedBuffer;
  		EAudioMixerStreamDataFormat::Type DataFormat;
- 		FThreadSafeBool bIsReady;
 
+		static size_t GetSizeForDataFormat(EAudioMixerStreamDataFormat::Type InDataFormat);
 		int32 CallCounterMixNextBuffer{ 0 };
 	};
 
@@ -535,10 +528,10 @@ namespace Audio
 		void StopGeneratingAudio();
 
 		/** Performs buffer fades for shutdown/startup of audio mixer. */
-		void ApplyMasterAttenuation();
+		void ApplyMasterAttenuation(TArrayView<const uint8>& InOutPoppedAudio);
 
 		template<typename BufferType>
-		void ApplyAttenuationInternal(BufferType* BufferDataPtr, const int32 NumFrames);
+		void ApplyAttenuationInternal(TArrayView<BufferType>& InOutBuffer);
 
 		/** When called, spins up a thread to start consuming output when no audio device is available. */
 		void StartRunningNullDevice();
@@ -553,10 +546,7 @@ namespace Audio
 		FAudioMixerOpenStreamParams OpenStreamParams;
 
 		/** List of generated output buffers. */
-		TArray<FOutputBuffer> OutputBuffers;
-
-		/** Special empty buffer for buffer underruns. */
-		FOutputBuffer UnderrunBuffer;
+		Audio::FOutputBuffer OutputBuffer;
 
 		/** Whether or not we warned of buffer underrun. */
 		bool bWarnedBufferUnderrun;
@@ -576,12 +566,6 @@ namespace Audio
 		/** Event allows you to block until fadeout is complete. */
 		FEvent* AudioFadeEvent;
 
-		/** The buffer which is currently submitted to the output device (and is being read from). */
-		TAtomic<int32> CurrentBufferReadIndex;
-
-		/** The buffer which is currently being rendered to (or about to be rendered to). */
-		TAtomic<int32> CurrentBufferWriteIndex;
-
 		/** The number of mixer buffers to queue on the output source voice. */
 		int32 NumOutputBuffers;
 
@@ -597,9 +581,6 @@ namespace Audio
 		/** String containing the last generated error. */
 		FString LastError;
 
-		/** Struct used to store render time analysis data. */
-		FAudioRenderTimeAnalysis RenderTimeAnalysis;
-
 		int32 CallCounterApplyAttenuationInternal{ 0 };
 		int32 CallCounterReadNextBuffer{ 0 };
 
diff --git a/Engine/Source/Runtime/NonRealtimeAudioRenderer/Private/AudioMixerPlatformNonRealtime.cpp b/Engine/Source/Runtime/NonRealtimeAudioRenderer/Private/AudioMixerPlatformNonRealtime.cpp
--- a/Engine/Source/Runtime/NonRealtimeAudioRenderer/Private/AudioMixerPlatformNonRealtime.cpp
+++ b/Engine/Source/Runtime/NonRealtimeAudioRenderer/Private/AudioMixerPlatformNonRealtime.cpp
@@ -62,7 +62,7 @@ namespace Audio
 
 	void FMixerPlatformNonRealtime::RenderAudio(double NumSecondsToRender)
 	{
-		if (!bIsInitialized || !bIsDeviceOpen || OutputBuffers.Num() == 0)
+		if (!bIsInitialized || !bIsDeviceOpen)
 		{
 			return;
 		}
@@ -72,25 +72,14 @@ namespace Audio
 		// Increment how much audio time the user wants to have been rendered. 
 		TotalDesiredRender += NumSecondsToRender;
 
-		CurrentBufferWriteIndex = 0;
-		CurrentBufferReadIndex = 0;
-
 		// Keep rendering audio until we surpass their desired time, TimePerCallback may be much smaller than NumSecondsToRender.
 		while (TotalDurationRendered < TotalDesiredRender)
 		{
-			// RenderTimeAnalysis.Start();
-			if (OutputBuffers.IsValidIndex(CurrentBufferWriteIndex))
-			{
-				OutputBuffers[CurrentBufferWriteIndex].MixNextBuffer();
-			}
-			// RenderTimeAnalysis.End();
+			OutputBuffer.MixNextBuffer();
 
 			ReadNextBuffer();
 			TotalDurationRendered += TimePerCallback;
 		}
-
-		CurrentBufferReadIndex = INDEX_NONE;
-		CurrentBufferWriteIndex = INDEX_NONE;
 	}
 
 	void FMixerPlatformNonRealtime::OpenFileToWriteAudioTo(const FString& OutPath)
@@ -297,8 +286,10 @@ namespace Audio
 
 	void FMixerPlatformNonRealtime::ResumePlaybackOnNewDevice()
 	{
-		SubmitBuffer(OutputBuffers[CurrentBufferReadIndex].GetBufferData());
-		check(OpenStreamParams.NumFrames * AudioStreamInfo.DeviceInfo.NumChannels == OutputBuffers[CurrentBufferReadIndex].GetBuffer().Num());
+		int32 PoppedSampleCount;
+		TArrayView<const uint8> PoppedAudio = OutputBuffer.PopBufferData(PoppedSampleCount);
+		SubmitBuffer(PoppedAudio.GetData());
+		check(OpenStreamParams.NumFrames * AudioStreamInfo.DeviceInfo.NumChannels == OutputBuffer.GetNumSamples());
 
 		AudioRenderEvent->Trigger();
 	}
diff --git a/Engine/Source/Runtime/SignalProcessing/Public/DSP/Dsp.h b/Engine/Source/Runtime/SignalProcessing/Public/DSP/Dsp.h
--- a/Engine/Source/Runtime/SignalProcessing/Public/DSP/Dsp.h
+++ b/Engine/Source/Runtime/SignalProcessing/Public/DSP/Dsp.h
@@ -519,12 +519,12 @@ namespace Audio
 	 * Designed to be thread safe for SPSC; However, if Push() and Pop() are both trying to access an overlapping area of the buffer,
 	 * One of the calls will be truncated. Thus, it is advised that you use a high enough capacity that the producer and consumer are never in contention.
 	 */
-	template <typename SampleType>
+	template <typename SampleType, size_t Alignment = 16>
 	class TCircularAudioBuffer
 	{
 	private:
 
-		TArray<SampleType> InternalBuffer;
+		TArray<SampleType, TAlignedHeapAllocator<Alignment>> InternalBuffer;
 		uint32 Capacity;
 		FThreadSafeCounter ReadCounter;
 		FThreadSafeCounter WriteCounter;
diff --git a/Engine/Source/Runtime/Windows/AudioMixerXAudio2/Private/AudioMixerPlatformXAudio2.cpp b/Engine/Source/Runtime/Windows/AudioMixerXAudio2/Private/AudioMixerPlatformXAudio2.cpp
--- a/Engine/Source/Runtime/Windows/AudioMixerXAudio2/Private/AudioMixerPlatformXAudio2.cpp
+++ b/Engine/Source/Runtime/Windows/AudioMixerXAudio2/Private/AudioMixerPlatformXAudio2.cpp
@@ -1200,12 +1200,6 @@ namespace Audio
 			XAUDIO2_RETURN_ON_FAIL(XAudio2System->CreateSourceVoice(&OutputAudioStreamSourceVoice, &Format, XAUDIO2_VOICE_NOPITCH, 2.0f, &OutputVoiceCallback));
 
 			const int32 NewNumSamples = OpenStreamParams.NumFrames * AudioStreamInfo.DeviceInfo.NumChannels;
-
-			// Clear the output buffers with zero's and submit one
-			for (int32 Index = 0; Index < OutputBuffers.Num(); ++Index)
-			{
-				OutputBuffers[Index].Reset(NewNumSamples);
-			}
 		}
 		else
 		{	
@@ -1222,11 +1216,11 @@ namespace Audio
 	{
 		if (OutputAudioStreamSourceVoice)
 		{
-			CurrentBufferReadIndex = 0;
-			CurrentBufferWriteIndex = 1;
+			int32 NumSamplesPopped = 0;
+			TArrayView<const uint8> PoppedAudio = OutputBuffer.PopBufferData(NumSamplesPopped);
+			SubmitBuffer(PoppedAudio.GetData());
 
-			SubmitBuffer(OutputBuffers[CurrentBufferReadIndex].GetBufferData());
-			check(OpenStreamParams.NumFrames * AudioStreamInfo.DeviceInfo.NumChannels == OutputBuffers[CurrentBufferReadIndex].GetBuffer().Num());
+			check(OpenStreamParams.NumFrames * AudioStreamInfo.DeviceInfo.NumChannels == OutputBuffer.GetNumSamples());
 
 			AudioRenderEvent->Trigger();
 
