Temporarily disable threaded device swap to help with build stability until we track down the issue.

diff --git a/Engine/Source/Runtime/AudioMixerCore/Private/AudioMixer.cpp b/Engine/Source/Runtime/AudioMixerCore/Private/AudioMixer.cpp
--- a/Engine/Source/Runtime/AudioMixerCore/Private/AudioMixer.cpp
+++ b/Engine/Source/Runtime/AudioMixerCore/Private/AudioMixer.cpp
@@ -60,7 +60,7 @@ FAutoConsoleVariableRef CVarDisableDeviceSwap(
 	TEXT("0: Not Enabled, 1: Enabled"),
 	ECVF_Default);
 
-static int32 bUseThreadedDeviceSwapCVar = 1;
+static int32 bUseThreadedDeviceSwapCVar = 0;
 FAutoConsoleVariableRef CVarUseThreadedDeviceSwap(
 	TEXT("au.UseThreadedDeviceSwap"),
 	bUseThreadedDeviceSwapCVar,
