Fix for nullptr on partial init teardown.

diff --git a/Engine/Source/Runtime/AudioMixerCore/Private/AudioMixer.cpp b/Engine/Source/Runtime/AudioMixerCore/Private/AudioMixer.cpp
--- a/Engine/Source/Runtime/AudioMixerCore/Private/AudioMixer.cpp
+++ b/Engine/Source/Runtime/AudioMixerCore/Private/AudioMixer.cpp
@@ -460,7 +460,11 @@ namespace Audio
 		float SafeSampleRate = OpenStreamParams.SampleRate > 0.f ? OpenStreamParams.SampleRate : 48000.f;
 		float BufferDuration = ((float)OpenStreamParams.NumFrames) / SafeSampleRate;
 
-		AudioRenderEvent->Trigger();
+		if (AudioRenderEvent)
+		{
+			AudioRenderEvent->Trigger();
+		}
+
 		if (!NullDeviceCallback.IsValid())
 		{
 			// Create the thread and tell it not to pause.
