Added au.RenderThreadAffinity cvar to be able to override AudioMixerRenderThread affinity if needed

#rnx

diff --git a/Engine/Source/Runtime/AudioMixerCore/Private/AudioMixer.cpp b/Engine/Source/Runtime/AudioMixerCore/Private/AudioMixer.cpp
--- a/Engine/Source/Runtime/AudioMixerCore/Private/AudioMixer.cpp
+++ b/Engine/Source/Runtime/AudioMixerCore/Private/AudioMixer.cpp
@@ -44,6 +44,15 @@ FAutoConsoleVariableRef CVarLogRenderThreadPriority(
 	TEXT("0: Normal, 1: Above Normal, 2: Below Normal, 3: Highest, 4: Lowest, 5: Slightly Below Normal, 6: Time Critical"),
 	ECVF_Default);
 
+static int32 SetRenderThreadAffinityCVar = 0;
+FAutoConsoleVariableRef CVarRenderThreadAffinity(
+	TEXT("au.RenderThreadAffinity"),
+	SetRenderThreadAffinityCVar,
+	TEXT("Override audio render thread affinity.\n")
+	TEXT("0: Disabled (Default), otherwise overriden thread affinity."),
+	ECVF_Default);
+
+
 static int32 EnableDetailedWindowsDeviceLoggingCVar = 0;
 FAutoConsoleVariableRef CVarEnableDetailedWindowsDeviceLogging(
 	TEXT("au.EnableDetailedWindowsDeviceLogging"),
@@ -632,7 +641,8 @@ namespace Audio
 		check(AudioFadeEvent != nullptr);
 
 		check(!AudioRenderThread.IsValid());
-		AudioRenderThread.Reset(FRunnableThread::Create(this, *FString::Printf(TEXT("AudioMixerRenderThread(%d)"), AudioMixerTaskCounter.Increment()), 0, (EThreadPriority)SetRenderThreadPriorityCVar, FPlatformAffinity::GetAudioThreadMask()));
+		uint64 RenderThreadAffinityCVar = SetRenderThreadAffinityCVar > 0 ? uint64(SetRenderThreadAffinityCVar) : FPlatformAffinity::GetAudioThreadMask();
+		AudioRenderThread.Reset(FRunnableThread::Create(this, *FString::Printf(TEXT("AudioMixerRenderThread(%d)"), AudioMixerTaskCounter.Increment()), 0, (EThreadPriority)SetRenderThreadPriorityCVar, RenderThreadAffinityCVar));
 		check(AudioRenderThread.IsValid());
 	}
 
