PR #8070: Fix copy paste in audio mixer thread priority (Contributed by [REDACTED_PERSON])

diff --git a/Engine/Source/Runtime/AudioMixerCore/Private/AudioMixer.cpp b/Engine/Source/Runtime/AudioMixerCore/Private/AudioMixer.cpp
--- a/Engine/Source/Runtime/AudioMixerCore/Private/AudioMixer.cpp
+++ b/Engine/Source/Runtime/AudioMixerCore/Private/AudioMixer.cpp
@@ -37,9 +37,9 @@ FAutoConsoleVariableRef CVarMinTimeBetweenUnderrunWarningsMs(
 
 // Command for setting the audio render thread priority.
 static int32 SetRenderThreadPriorityCVar = (int32)TPri_Highest;
-FAutoConsoleVariableRef CVarLogRenderThreadPriority(
+FAutoConsoleVariableRef CVarSetRenderThreadPriority(
 	TEXT("au.RenderThreadPriority"),
-	LogRenderTimesCVar,
+	SetRenderThreadPriorityCVar,
 	TEXT("Sets audio render thread priority. Defaults to 3.\n")
 	TEXT("0: Normal, 1: Above Normal, 2: Below Normal, 3: Highest, 4: Lowest, 5: Slightly Below Normal, 6: Time Critical"),
 	ECVF_Default);
