Inclusive language push, removing master/slave language from Audio Mixer + proper deprecations

diff --git a/Engine/Source/Runtime/AudioMixerCore/Private/AudioMixer.cpp b/Engine/Source/Runtime/AudioMixerCore/Private/AudioMixer.cpp
--- a/Engine/Source/Runtime/AudioMixerCore/Private/AudioMixer.cpp
+++ b/Engine/Source/Runtime/AudioMixerCore/Private/AudioMixer.cpp
@@ -514,7 +514,7 @@ namespace Audio
 		NullDeviceCallback.Reset(new FMixerNullCallback(InBufferDuration, InCallback, TPri_TimeCritical, bShouldPauseOnStart));
 	}
 
-	void IAudioMixerPlatformInterface::ApplyMasterAttenuation(TArrayView<const uint8>& OutPoppedAudio)
+	void IAudioMixerPlatformInterface::ApplyPrimaryAttenuation(TArrayView<const uint8>& OutPoppedAudio)
 	{
 		EAudioMixerStreamDataFormat::Type Format = OutputBuffer.GetFormat();
 
@@ -599,7 +599,7 @@ namespace Audio
 			bWarnedBufferUnderrun = false;
 		}
 
-		ApplyMasterAttenuation(PoppedAudio);
+		ApplyPrimaryAttenuation(PoppedAudio);
 		SubmitBuffer(PoppedAudio.GetData());
 
 		DeviceSwapCriticalSection.Unlock();
diff --git a/Engine/Source/Runtime/AudioMixerCore/Public/AudioMixer.h b/Engine/Source/Runtime/AudioMixerCore/Public/AudioMixer.h
--- a/Engine/Source/Runtime/AudioMixerCore/Public/AudioMixer.h
+++ b/Engine/Source/Runtime/AudioMixerCore/Public/AudioMixer.h
@@ -631,9 +631,13 @@ namespace Audio
 		/** Stops the render thread from generating audio. */
 		void StopGeneratingAudio();
 
-		/** Performs buffer fades for shutdown/startup of audio mixer. */
+		// Deprecated - use ApplyPrimaryAttenuation
+		UE_DEPRECATED(5.1, "ApplyMasterAttenuation is deprecated, please use ApplyPrimaryAttenuation instead.")
 		void ApplyMasterAttenuation(TArrayView<const uint8>& InOutPoppedAudio);
 
+		/** Performs buffer fades for shutdown/startup of audio mixer. */
+		void ApplyPrimaryAttenuation(TArrayView<const uint8>& InOutPoppedAudio);
+
 		template<typename BufferType>
 		void ApplyAttenuationInternal(TArrayView<BufferType>& InOutBuffer);
 
@@ -677,7 +681,7 @@ namespace Audio
 		/** The number of mixer buffers to queue on the output source voice. */
 		int32 NumOutputBuffers;
 
-		/** The fade value. Used for fading in/out master audio. */
+		/** The fade value. Used for fading in/out primary audio. */
 		float FadeVolume;
 
 		/** Source param used to fade in and out audio device. */
diff --git a/Engine/Source/Runtime/SignalProcessing/Public/DSP/Osc.h b/Engine/Source/Runtime/SignalProcessing/Public/DSP/Osc.h
--- a/Engine/Source/Runtime/SignalProcessing/Public/DSP/Osc.h
+++ b/Engine/Source/Runtime/SignalProcessing/Public/DSP/Osc.h
@@ -128,6 +128,10 @@ namespace Audio
 		// Sets whether or not this oscillator should be synced to a leader oscillator. leader oscillator needs to have set this oscillator as its follower.
 		void SetSync(const bool bInSync) { bIsSync = bInSync; }
 
+		// Deprecated: use SetFollowerOsc
+		UE_DEPRECATED(5.1, "SetSlaveOsc is deprecated, please use SetFollowerOsc instead.")
+		void SetSlaveOsc(IOscBase* InSlaveOsc);
+
 		// Sets the input oscillator as the follower of this oscillator
 		void SetFollowerOsc(IOscBase* InFollowerOsc);
 
