Disable Device Swapping on hardware timeouts when we've got debugger attached. Prevents needless swaps/spam when stepping in debugger.

diff --git a/Engine/Source/Runtime/AudioMixerCore/Private/AudioMixer.cpp b/Engine/Source/Runtime/AudioMixerCore/Private/AudioMixer.cpp
--- a/Engine/Source/Runtime/AudioMixerCore/Private/AudioMixer.cpp
+++ b/Engine/Source/Runtime/AudioMixerCore/Private/AudioMixer.cpp
@@ -751,16 +751,19 @@ namespace Audio
 			// Bounds check the timeout for our audio render event.
 			OverrunTimeoutCVar = FMath::Clamp(OverrunTimeoutCVar, 500, 5000);
 
+			// If we're debugging, make the timeout the maximum to avoid needless swaps.
+			OverrunTimeoutCVar = FPlatformMisc::IsDebuggerPresent() ? TNumericLimits<uint32>::Max() : OverrunTimeoutCVar;
+
 			// Now wait for a buffer to be consumed, which will bump up the read index.
-			double WaitStartTime = FPlatformTime::Seconds();
+			const double WaitStartTime = FPlatformTime::Seconds();
 			if (AudioRenderEvent && !AudioRenderEvent->Wait(static_cast<uint32>(OverrunTimeoutCVar)))
 			{
 				// if we reached this block, we timed out, and should attempt to
 				// bail on our current device.
 				RequestDeviceSwap(TEXT(""), /* force */true, TEXT("AudioMixerPlatformInterface. Timeout waiting for h/w."));
 
-				float TimeWaited = FPlatformTime::Seconds() - WaitStartTime;
-				UE_LOG(LogAudioMixer, Warning, TEXT("AudioMixerPlatformInterface Timeout [%dms] waiting for h/w. InstanceID=%d"), (int32)(TimeWaited * 1000.f),InstanceID);
+				const float TimeWaited = FPlatformTime::Seconds() - WaitStartTime;
+				UE_LOG(LogAudioMixer, Warning, TEXT("AudioMixerPlatformInterface Timeout [%2.f Seconds] waiting for h/w. InstanceID=%d"), TimeWaited,InstanceID);
 			}
 		}
 
