Alliow AudioCallbackBufferFrameSize to be configured as low as 240.

diff --git a/Engine/Source/Runtime/AudioMixerCore/Private/AudioMixer.cpp b/Engine/Source/Runtime/AudioMixerCore/Private/AudioMixer.cpp
--- a/Engine/Source/Runtime/AudioMixerCore/Private/AudioMixer.cpp
+++ b/Engine/Source/Runtime/AudioMixerCore/Private/AudioMixer.cpp
@@ -934,7 +934,7 @@ FAudioPlatformSettings FAudioPlatformSettings::GetPlatformSettings(const TCHAR*
 
 	if (GConfig->GetString(PlatformSettingsConfigFile, TEXT("AudioCallbackBufferFrameSize"), TempString, GEngineIni))
 	{
-		Settings.CallbackBufferFrameSize = FMath::Max(FCString::Atoi(*TempString), 256);
+		Settings.CallbackBufferFrameSize = FMath::Max(FCString::Atoi(*TempString), 240);
 	}
 
 	if (GConfig->GetString(PlatformSettingsConfigFile, TEXT("AudioNumBuffersToEnqueue"), TempString, GEngineIni))
