Audio Diagnositcs. Adding CVAr to deliberately drop hardware callbacks to diagnose hardware timeouts.

diff --git a/Engine/Source/Runtime/AudioMixerCore/Private/AudioMixer.cpp b/Engine/Source/Runtime/AudioMixerCore/Private/AudioMixer.cpp
--- a/Engine/Source/Runtime/AudioMixerCore/Private/AudioMixer.cpp
+++ b/Engine/Source/Runtime/AudioMixerCore/Private/AudioMixer.cpp
@@ -124,6 +124,14 @@ FAutoConsoleVariableRef ExtraAudioMixerDeviceLogging(
 	TEXT("0: no logging, 1: logging every 500 callbacks \n"),
 	ECVF_Default);
 
+static int32 AudioMixerDebugForceDroppedHardwareCallbackCVar = 0;
+FAutoConsoleVariableRef CVarAudioMixerDebugForceDroppedHardwareCallback(
+	TEXT("au.debug.ForceDroppedHardwareCallback"),
+	AudioMixerDebugForceDroppedHardwareCallbackCVar,
+	TEXT("Will drop the next N hardware callbacks\n")
+	TEXT("N == 0 off, N < 0 (will drop all callbacks), N > 0 will drop (will drop N callbacks)"),
+	ECVF_Cheat);
+
 // Stat definitions for profiling audio mixer 
 DEFINE_STAT(STAT_AudioMixerRenderAudio);
 
@@ -531,6 +539,19 @@ namespace Audio
 	{
 		LLM_SCOPE(ELLMTag::AudioMixer);
 
+#if !UE_BUILD_SHIPPING
+		// Debugging CVAR to drop hardware calls.
+		if (AudioMixerDebugForceDroppedHardwareCallbackCVar != 0)
+		{
+			// Drop N calls...
+			if (AudioMixerDebugForceDroppedHardwareCallbackCVar > 0)
+			{
+				AudioMixerDebugForceDroppedHardwareCallbackCVar--;
+			}
+			return;
+		}
+#endif //!UE_BUILD_SHIPPING
+
 		// If we are flushing buffers for our output voice and this is being called on the audio thread directly,
 		// early exit.
 		if (bIsInDeviceSwap)
