Fix up accounting for circular buffer

diff --git a/Engine/Source/Runtime/AudioMixerCore/Private/AudioMixer.cpp b/Engine/Source/Runtime/AudioMixerCore/Private/AudioMixer.cpp
--- a/Engine/Source/Runtime/AudioMixerCore/Private/AudioMixer.cpp
+++ b/Engine/Source/Runtime/AudioMixerCore/Private/AudioMixer.cpp
@@ -220,7 +220,8 @@ namespace Audio
 	bool FOutputBuffer::MixNextBuffer()
  	{
 		// If the circular queue is already full, exit.
-		if (CircularBuffer.Remainder() < static_cast<uint32>(RenderBuffer.Num()))
+		const int32 SpaceLeftInBufferInSamples = CircularBuffer.Remainder() / GetSizeForDataFormat(DataFormat);
+		if (SpaceLeftInBufferInSamples < RenderBuffer.Num())
 		{
 			return false;
 		}
diff --git a/Engine/Source/Runtime/SignalProcessing/Public/DSP/Dsp.h b/Engine/Source/Runtime/SignalProcessing/Public/DSP/Dsp.h
--- a/Engine/Source/Runtime/SignalProcessing/Public/DSP/Dsp.h
+++ b/Engine/Source/Runtime/SignalProcessing/Public/DSP/Dsp.h
@@ -1165,10 +1165,7 @@ namespace Audio
 		// Get number of samples that can be pushed onto the buffer before it is full.
 		uint32 Remainder() const
 		{
-			const uint32 ReadIndex = ReadCounter.GetValue();
-			const uint32 WriteIndex = WriteCounter.GetValue();
-
-			return (Capacity - 1 - WriteIndex + ReadIndex) % Capacity;
+			return Capacity - Num() - 1;
 		}
 	};
 
