Protect against null ptr.

diff --git a/Engine/Source/Runtime/AudioMixerCore/Private/AudioMixer.cpp b/Engine/Source/Runtime/AudioMixerCore/Private/AudioMixer.cpp
--- a/Engine/Source/Runtime/AudioMixerCore/Private/AudioMixer.cpp
+++ b/Engine/Source/Runtime/AudioMixerCore/Private/AudioMixer.cpp
@@ -442,8 +442,10 @@ namespace Audio
 
 			bFadedOut = (FadeVolume == 0.0f);
 			bPerformingFade = false;
-			AudioFadeEvent->Trigger();
-
+			if (AudioFadeEvent)
+			{
+				AudioFadeEvent->Trigger();
+			}
 			if (bLog)
 			{
 				UE_LOG(LogAudioMixer, Display, TEXT("IAudioMixerPlatformInterface::ApplyAttenuationInternal() Faded from %f to %f"), FadeVolume, FadeParam.GetValue());
