Fixed the audio mixer event timeout duration changing to minimum if you attach and then detach a debugger.

1. The default is 5000ms and the CVar is clamped to the range [500, 5000].
2. Attaching a debugger sets it to -1 (infinite).
3. Detaching the debugger then causes it to clamp as int32 to 500ms.

Change the clamp to uint32 but also don't actually change the value of the CVar, just the amount we're sleeping for this iteration.

diff --git a/Engine/Source/Runtime/AudioMixerCore/Private/AudioMixer.cpp b/Engine/Source/Runtime/AudioMixerCore/Private/AudioMixer.cpp
--- a/Engine/Source/Runtime/AudioMixerCore/Private/AudioMixer.cpp
+++ b/Engine/Source/Runtime/AudioMixerCore/Private/AudioMixer.cpp
@@ -747,15 +747,12 @@ namespace Audio
 			{
 			}
 
-			// Bounds check the timeout for our audio render event.
-			OverrunTimeoutCVar = FMath::Clamp(OverrunTimeoutCVar, 500, 5000);
-
-			// If we're debugging, make the timeout the maximum to avoid needless swaps.
-			OverrunTimeoutCVar = FPlatformMisc::IsDebuggerPresent() ? TNumericLimits<uint32>::Max() : OverrunTimeoutCVar;
+			// Bounds check the timeout for our audio render event, waiting indefinitely if a debugger is attached
+			const uint32 CurrentWaitTimeout = FPlatformMisc::IsDebuggerPresent() ? MAX_uint32 : FMath::Clamp<uint32>(OverrunTimeoutCVar, 500, 5000);
 
 			// Now wait for a buffer to be consumed, which will bump up the read index.
 			const double WaitStartTime = FPlatformTime::Seconds();
-			if (AudioRenderEvent && !AudioRenderEvent->Wait(static_cast<uint32>(OverrunTimeoutCVar)))
+			if (AudioRenderEvent && !AudioRenderEvent->Wait(CurrentWaitTimeout))
 			{
 				// if we reached this block, we timed out, and should attempt to
 				// bail on our current device.
