[AudioMixerCore]
* Fixed c++modules/headerunits compile errors

diff --git a/Engine/Source/Runtime/AudioMixerCore/Public/AudioDefines.h b/Engine/Source/Runtime/AudioMixerCore/Public/AudioDefines.h
--- a/Engine/Source/Runtime/AudioMixerCore/Public/AudioDefines.h
+++ b/Engine/Source/Runtime/AudioMixerCore/Public/AudioDefines.h
@@ -67,8 +67,8 @@ namespace Audio
 
 
 	// Magic numbers for mono upmix method
-	static constexpr float MonoUpmixFullVolume = 1.0f;
-	static constexpr float MonoUpmixEqualPower = 0.707f;
-	static constexpr float MonoUpmixLinear = 0.5f;
+	inline constexpr float MonoUpmixFullVolume = 1.0f;
+	inline constexpr float MonoUpmixEqualPower = 0.707f;
+	inline constexpr float MonoUpmixLinear = 0.5f;
 
 }
\ No newline at end of file
