AudioMixerCore put string definitions for Insights events into .cpp file

When it's in the header like that, we get init code in every TU.

diff --git a/Engine/Plugins/AudioGameplay/Source/AudioGameplay/AudioGameplay.Build.cs b/Engine/Plugins/AudioGameplay/Source/AudioGameplay/AudioGameplay.Build.cs
--- a/Engine/Plugins/AudioGameplay/Source/AudioGameplay/AudioGameplay.Build.cs
+++ b/Engine/Plugins/AudioGameplay/Source/AudioGameplay/AudioGameplay.Build.cs
@@ -23,6 +23,7 @@ public class AudioGameplay : ModuleRules
 			{
 				"CoreUObject",
 				"Engine",
+				"AudioMixerCore",
 			}
 			);
 		
diff --git a/Engine/Source/Editor/UnrealEd/UnrealEd.Build.cs b/Engine/Source/Editor/UnrealEd/UnrealEd.Build.cs
--- a/Engine/Source/Editor/UnrealEd/UnrealEd.Build.cs
+++ b/Engine/Source/Editor/UnrealEd/UnrealEd.Build.cs
@@ -190,6 +190,7 @@ public class UnrealEd : ModuleRules
 				"GeometryCore",
 				"UniversalObjectLocatorEditor",
 				"Renderer",
+				"AudioMixerCore",
 			}
 		);
 
diff --git a/Engine/Source/Runtime/AudioMixerCore/Private/AudioMixer.cpp b/Engine/Source/Runtime/AudioMixerCore/Private/AudioMixer.cpp
--- a/Engine/Source/Runtime/AudioMixerCore/Private/AudioMixer.cpp
+++ b/Engine/Source/Runtime/AudioMixerCore/Private/AudioMixer.cpp
@@ -29,6 +29,39 @@ UE_TRACE_CHANNEL_DEFINE(AudioChannel);
 UE_TRACE_CHANNEL_DEFINE(AudioMixerChannel);
 #endif // UE_AUDIO_PROFILERTRACE_ENABLED
 
+// Audio Insights Event Log in-built event types
+namespace Audio::Trace::EventLog::ID
+{
+	// Active Sound : Playing Audio
+	const FString SoundStart = NSLOCTEXT("AudioInsightsEventLog", "EventLogTraceMessage_Playing", "Playing").ToString();
+	const FString SoundStop = NSLOCTEXT("AudioInsightsEventLog", "EventLogTraceMessage_Stopped", "Stopped").ToString();
+
+	// Virtualization
+	const FString SoundVirtualized = NSLOCTEXT("AudioInsightsEventLog", "EventLogTraceMessage_Virtualize", "Virtualized").ToString();
+	const FString SoundRealized = NSLOCTEXT("AudioInsightsEventLog", "EventLogTraceMessage_Realize", "Realized").ToString();
+
+	// Triggers
+	const FString PlayRequestSoundHandle = NSLOCTEXT("AudioInsightsEventLog", "EventLogTraceMessage_PlayRequestedSoundHandle", "Play Request : Sound Handle").ToString();
+	const FString StopRequestedSoundHandle = NSLOCTEXT("AudioInsightsEventLog", "EventLogTraceMessage_StopRequestedSoundHandle", "Stop Request : Sound Handle").ToString();
+	
+	const FString PlayRequestAudioComponent = NSLOCTEXT("AudioInsightsEventLog", "EventLogTraceMessage_PlayRequestedAudioComponent", "Play Request : Audio Component").ToString();
+	const FString StopRequestAudioComponent = NSLOCTEXT("AudioInsightsEventLog", "EventLogTraceMessage_StopRequestedAudioComponent", "Stop Request : Audio Component").ToString();
+	
+	const FString PlayRequestOneShot = NSLOCTEXT("AudioInsightsEventLog", "EventLogTraceMessage_PlayRequestedOneShot", "Play Request : One shot").ToString();
+	const FString PlayRequestSoundAtLocation = NSLOCTEXT("AudioInsightsEventLog", "EventLogTraceMessage_PlayRequestedSoundAtLocation", "Play Request : Sound at location").ToString();
+	const FString PlayRequestSound2D = NSLOCTEXT("AudioInsightsEventLog", "EventLogTraceMessage_PlayRequested2DSound", "Play Request : Play Sound 2D").ToString();
+	const FString PlayRequestSlateSound = NSLOCTEXT("AudioInsightsEventLog", "EventLogTraceMessage_PlayRequestedSlateSound", "Play Request : Slate Sound").ToString();
+	
+	const FString StopRequestActiveSound = NSLOCTEXT("AudioInsightsEventLog", "EventLogTraceMessage_StopRequestedActiveSound", "Stop Request : Active Sound").ToString();
+	const FString StopRequestSoundsUsingResource = NSLOCTEXT("AudioInsightsEventLog", "EventLogTraceMessage_StopRequestedSoundsUsingResource", "Stop Request : Sounds using resource").ToString();
+	const FString StopRequestConcurrency = NSLOCTEXT("AudioInsightsEventLog", "EventLogTraceMessage_StopRequestedConcurrency", "Stop Request : Concurrency").ToString();
+	
+	const FString PauseSoundRequested = NSLOCTEXT("AudioInsightsEventLog", "EventLogTraceMessage_PauseSoundRequested", "Paused").ToString();
+	const FString ResumeSoundRequested = NSLOCTEXT("AudioInsightsEventLog", "EventLogTraceMessage_ResumeSoundRequested", "Resumed").ToString();
+	
+	const FString StopAllRequested = NSLOCTEXT("AudioInsightsEventLog", "EventLogTraceMessage_StopAllRequested", "Stop All").ToString();
+	const FString FlushAudioDeviceRequested = NSLOCTEXT("AudioInsightsEventLog", "EventLogTraceMessage_FlushAudioDeviceRequested", "Flush Audio Device").ToString();
+}
 
 // Command to enable logging to display accurate audio render times
 static int32 LogRenderTimesCVar = 0;
diff --git a/Engine/Source/Runtime/AudioMixerCore/Public/AudioMixerTrace.h b/Engine/Source/Runtime/AudioMixerCore/Public/AudioMixerTrace.h
--- a/Engine/Source/Runtime/AudioMixerCore/Public/AudioMixerTrace.h
+++ b/Engine/Source/Runtime/AudioMixerCore/Public/AudioMixerTrace.h
@@ -25,35 +25,36 @@
 
 
 // Audio Insights Event Log in-built event types
+// actual event list is in AudioMixer.cpp
 namespace Audio::Trace::EventLog::ID
 {
 	// Active Sound : Playing Audio
-	const FString SoundStart = NSLOCTEXT("AudioInsightsEventLog", "EventLogTraceMessage_Playing", "Playing").ToString();
-	const FString SoundStop = NSLOCTEXT("AudioInsightsEventLog", "EventLogTraceMessage_Stopped", "Stopped").ToString();
+	extern AUDIOMIXERCORE_API const FString SoundStart;
+	extern AUDIOMIXERCORE_API const FString SoundStop;
 
 	// Virtualization
-	const FString SoundVirtualized = NSLOCTEXT("AudioInsightsEventLog", "EventLogTraceMessage_Virtualize", "Virtualized").ToString();
-	const FString SoundRealized = NSLOCTEXT("AudioInsightsEventLog", "EventLogTraceMessage_Realize", "Realized").ToString();
+	extern AUDIOMIXERCORE_API const FString SoundVirtualized;
+	extern AUDIOMIXERCORE_API const FString SoundRealized;
 
 	// Triggers
-	const FString PlayRequestSoundHandle = NSLOCTEXT("AudioInsightsEventLog", "EventLogTraceMessage_PlayRequestedSoundHandle", "Play Request : Sound Handle").ToString();
-	const FString StopRequestedSoundHandle = NSLOCTEXT("AudioInsightsEventLog", "EventLogTraceMessage_StopRequestedSoundHandle", "Stop Request : Sound Handle").ToString();
+	extern AUDIOMIXERCORE_API const FString PlayRequestSoundHandle;
+	extern AUDIOMIXERCORE_API const FString StopRequestedSoundHandle;
 	
-	const FString PlayRequestAudioComponent = NSLOCTEXT("AudioInsightsEventLog", "EventLogTraceMessage_PlayRequestedAudioComponent", "Play Request : Audio Component").ToString();
-	const FString StopRequestAudioComponent = NSLOCTEXT("AudioInsightsEventLog", "EventLogTraceMessage_StopRequestedAudioComponent", "Stop Request : Audio Component").ToString();
+	extern AUDIOMIXERCORE_API const FString PlayRequestAudioComponent;
+	extern AUDIOMIXERCORE_API const FString StopRequestAudioComponent;
 	
-	const FString PlayRequestOneShot = NSLOCTEXT("AudioInsightsEventLog", "EventLogTraceMessage_PlayRequestedOneShot", "Play Request : One shot").ToString();
-	const FString PlayRequestSoundAtLocation = NSLOCTEXT("AudioInsightsEventLog", "EventLogTraceMessage_PlayRequestedSoundAtLocation", "Play Request : Sound at location").ToString();
-	const FString PlayRequestSound2D = NSLOCTEXT("AudioInsightsEventLog", "EventLogTraceMessage_PlayRequested2DSound", "Play Request : Play Sound 2D").ToString();
-	const FString PlayRequestSlateSound = NSLOCTEXT("AudioInsightsEventLog", "EventLogTraceMessage_PlayRequestedSlateSound", "Play Request : Slate Sound").ToString();
+	extern AUDIOMIXERCORE_API const FString PlayRequestOneShot;
+	extern AUDIOMIXERCORE_API const FString PlayRequestSoundAtLocation;
+	extern AUDIOMIXERCORE_API const FString PlayRequestSound2D;
+	extern AUDIOMIXERCORE_API const FString PlayRequestSlateSound;
 	
-	const FString StopRequestActiveSound = NSLOCTEXT("AudioInsightsEventLog", "EventLogTraceMessage_StopRequestedActiveSound", "Stop Request : Active Sound").ToString();
-	const FString StopRequestSoundsUsingResource = NSLOCTEXT("AudioInsightsEventLog", "EventLogTraceMessage_StopRequestedSoundsUsingResource", "Stop Request : Sounds using resource").ToString();
-	const FString StopRequestConcurrency = NSLOCTEXT("AudioInsightsEventLog", "EventLogTraceMessage_StopRequestedConcurrency", "Stop Request : Concurrency").ToString();
+	extern AUDIOMIXERCORE_API const FString StopRequestActiveSound;
+	extern AUDIOMIXERCORE_API const FString StopRequestSoundsUsingResource;
+	extern AUDIOMIXERCORE_API const FString StopRequestConcurrency;
 	
-	const FString PauseSoundRequested = NSLOCTEXT("AudioInsightsEventLog", "EventLogTraceMessage_PauseSoundRequested", "Paused").ToString();
-	const FString ResumeSoundRequested = NSLOCTEXT("AudioInsightsEventLog", "EventLogTraceMessage_ResumeSoundRequested", "Resumed").ToString();
+	extern AUDIOMIXERCORE_API const FString PauseSoundRequested;
+	extern AUDIOMIXERCORE_API const FString ResumeSoundRequested;
 	
-	const FString StopAllRequested = NSLOCTEXT("AudioInsightsEventLog", "EventLogTraceMessage_StopAllRequested", "Stop All").ToString();
-	const FString FlushAudioDeviceRequested = NSLOCTEXT("AudioInsightsEventLog", "EventLogTraceMessage_FlushAudioDeviceRequested", "Flush Audio Device").ToString();
-}
\ No newline at end of file
+	extern AUDIOMIXERCORE_API const FString StopAllRequested;
+	extern AUDIOMIXERCORE_API const FString FlushAudioDeviceRequested;
+}
