When selecting default device, take into account platform preference of device role.

#rnx

diff --git a/Engine/Source/Runtime/AudioDeviceEnumeration/Windows/WindowsMMDeviceEnumeration/Private/WindowsMMDeviceInfoCache.cpp b/Engine/Source/Runtime/AudioDeviceEnumeration/Windows/WindowsMMDeviceEnumeration/Private/WindowsMMDeviceInfoCache.cpp
--- a/Engine/Source/Runtime/AudioDeviceEnumeration/Windows/WindowsMMDeviceEnumeration/Private/WindowsMMDeviceInfoCache.cpp
+++ b/Engine/Source/Runtime/AudioDeviceEnumeration/Windows/WindowsMMDeviceEnumeration/Private/WindowsMMDeviceInfoCache.cpp
@@ -897,6 +897,17 @@ namespace Audio
 		return 0;
 	}
 
+	FName FWindowsMMDeviceCache::GetDefaultOutputDevice(EAudioDeviceRole InRole) const
+	{
+		FReadScopeLock MapReadLock(CacheMutationLock);
+
+		if (ensure(InRole >= EAudioDeviceRole::Console && InRole < EAudioDeviceRole::COUNT))
+		{
+			return DefaultRenderId[(int32)InRole];
+		}
+		return NAME_None;
+	}
+
 	FName FWindowsMMDeviceCache::GetDefaultOutputDevice_NoLock() const
 	{
 		if (!DefaultRenderId[(int32)EAudioDeviceRole::Console].IsNone())
diff --git a/Engine/Source/Runtime/AudioDeviceEnumeration/Windows/WindowsMMDeviceEnumeration/Public/WindowsMMDeviceInfoCache.h b/Engine/Source/Runtime/AudioDeviceEnumeration/Windows/WindowsMMDeviceEnumeration/Public/WindowsMMDeviceInfoCache.h
--- a/Engine/Source/Runtime/AudioDeviceEnumeration/Windows/WindowsMMDeviceEnumeration/Public/WindowsMMDeviceInfoCache.h
+++ b/Engine/Source/Runtime/AudioDeviceEnumeration/Windows/WindowsMMDeviceEnumeration/Public/WindowsMMDeviceInfoCache.h
@@ -110,6 +110,8 @@ struct FWindowsMMDeviceCache : IAudioMixerDeviceChangedListener, IAudioPlatformD
 
 	UE_API virtual void Refresh() override;
 	
+	UE_API virtual FName GetDefaultOutputDevice(EAudioDeviceRole InRole) const override;
+
 	UE_API FName GetDefaultOutputDevice_NoLock() const;
 
 	UE_API TOptional<FAudioPlatformDeviceInfo> FindDefaultOutputDevice() const override;
diff --git a/Engine/Source/Runtime/AudioMixerCore/Public/AudioMixer.h b/Engine/Source/Runtime/AudioMixerCore/Public/AudioMixer.h
--- a/Engine/Source/Runtime/AudioMixerCore/Public/AudioMixer.h
+++ b/Engine/Source/Runtime/AudioMixerCore/Public/AudioMixer.h
@@ -227,6 +227,25 @@ namespace Audio
 
 	enum class EDeviceEndpointType { Unknown, Render, Capture };
 
+	enum class EAudioDeviceRole
+	{
+		Console,
+		Multimedia,
+		Communications,
+
+		COUNT,
+	};
+
+	enum class EAudioDeviceState
+	{
+		Active,
+		Disabled,
+		NotPresent,
+		Unplugged,
+
+		COUNT,
+	};
+
 	// Interface for Caching Device Info.
 	class IAudioPlatformDeviceInfoCache
 	{
@@ -234,6 +253,7 @@ namespace Audio
 		// Pure Interface. 
 		virtual ~IAudioPlatformDeviceInfoCache() = default;
 			
+		virtual FName GetDefaultOutputDevice(EAudioDeviceRole InRole) const = 0;
 		virtual TOptional<FAudioPlatformDeviceInfo> FindActiveOutputDevice(FName InDeviceID) const = 0;
 		virtual TArray<FAudioPlatformDeviceInfo> GetAllActiveOutputDevices() const = 0;
 		virtual TOptional<FAudioPlatformDeviceInfo> FindDefaultOutputDevice() const = 0;
@@ -330,25 +350,6 @@ namespace Audio
 		}
 	};
 
-	enum class EAudioDeviceRole
-	{
-		Console,
-		Multimedia,
-		Communications,
-
-		COUNT,
-	};
-
-	enum class EAudioDeviceState
-	{
-		Active,
-		Disabled,
-		NotPresent,
-		Unplugged,
-
-		COUNT,
-	};
-
 	/** Struct used to store render time analysis data. */
 	struct FAudioRenderTimeAnalysis
 	{
diff --git a/Engine/Source/Runtime/Windows/AudioMixerWasapi/Private/AudioMixerWasapi.h b/Engine/Source/Runtime/Windows/AudioMixerWasapi/Private/AudioMixerWasapi.h
--- a/Engine/Source/Runtime/Windows/AudioMixerWasapi/Private/AudioMixerWasapi.h
+++ b/Engine/Source/Runtime/Windows/AudioMixerWasapi/Private/AudioMixerWasapi.h
@@ -111,6 +111,9 @@ namespace Audio
 		 *  FAudioPlatformDeviceInfo rather than looking it up via the requested device Id.
 		 */
 		bool InitDeviceSwapContextInternal(const FString& InRequestedDeviceID, const TCHAR* InReason, const TOptional<FAudioPlatformDeviceInfo>& InDeviceInfo);
+
+		/** Returns the device role to be used when swapping to the default device. */
+		virtual EAudioDeviceRole GetDefaultDeviceRole() const { return EAudioDeviceRole::Console; }
 		
 	private:
 
diff --git a/Engine/Source/Runtime/Windows/AudioMixerWasapi/Private/AudioMixerWasapiDeviceListener.cpp b/Engine/Source/Runtime/Windows/AudioMixerWasapi/Private/AudioMixerWasapiDeviceListener.cpp
--- a/Engine/Source/Runtime/Windows/AudioMixerWasapi/Private/AudioMixerWasapiDeviceListener.cpp
+++ b/Engine/Source/Runtime/Windows/AudioMixerWasapi/Private/AudioMixerWasapiDeviceListener.cpp
@@ -178,7 +178,17 @@ namespace Audio
 			// to use and is not interested in following the current system default device.
 			if (GetIsListeningForDeviceEvents())
 			{
-				RequestDeviceSwap(TEXT(""), /* force */ true, TEXT("FAudioMixerWasapi::OnDeviceStateChanged"));
+				FString DefaultDeviceId;
+				if (IAudioPlatformDeviceInfoCache* Cache = GetDeviceInfoCache())
+				{
+					const FName DeviceIdName = Cache->GetDefaultOutputDevice(GetDefaultDeviceRole());
+					if (DeviceIdName != NAME_None)
+					{
+						DefaultDeviceId = DeviceIdName.ToString();
+					}
+				}
+
+				RequestDeviceSwap(DefaultDeviceId, /* force */ true, TEXT("FAudioMixerWasapi::OnDeviceStateChanged"));
 			}
 		}
 		
