Move AudioMixerChannel enum into its own header

#tests LowLevelTests audio
#Synced-CL 48632724

diff --git a/Engine/Source/Runtime/AudioMixerCore/Public/AudioMixer.h b/Engine/Source/Runtime/AudioMixerCore/Public/AudioMixer.h
--- a/Engine/Source/Runtime/AudioMixerCore/Public/AudioMixer.h
+++ b/Engine/Source/Runtime/AudioMixerCore/Public/AudioMixer.h
@@ -3,6 +3,7 @@
 #pragma once
 
 #include "Async/Future.h"
+#include "AudioMixerChannel.h"
 #include "AudioMixerLog.h"
 #include "AudioMixerNullDevice.h"
 #include "AudioMixerTypes.h"
@@ -76,65 +77,6 @@ DECLARE_STATS_GROUP(TEXT("AudioMixer"), STATGROUP_AudioMixer, STATCAT_Advanced);
 // Tracks the time for the full render block 
 DECLARE_CYCLE_STAT_EXTERN(TEXT("Render Audio"), STAT_AudioMixerRenderAudio, STATGROUP_AudioMixer, AUDIOMIXERCORE_API);
 
-namespace EAudioMixerChannel
-{
-	/** Enumeration values represent sound file or speaker channel types. */
-	enum Type
-	{
-		FrontLeft,
-		FrontRight,
-		FrontCenter,
-		LowFrequency,
-		BackLeft,
-		BackRight,
-		FrontLeftOfCenter,
-		FrontRightOfCenter,
-		BackCenter,
-		SideLeft,
-		SideRight,
-		TopCenter,
-		TopFrontLeft,
-		TopFrontCenter,
-		TopFrontRight,
-		TopBackLeft,
-		TopBackCenter,
-		TopBackRight,
-		Unknown,
-		ChannelTypeCount,
-		DefaultChannel = FrontLeft
-	};
-
-	inline constexpr int32 MaxSupportedChannel = EAudioMixerChannel::TopCenter;
-
-	inline const TCHAR* ToString(EAudioMixerChannel::Type InType)
-	{
-		switch (InType)
-		{
-		case FrontLeft:				return TEXT("FrontLeft");
-		case FrontRight:			return TEXT("FrontRight");
-		case FrontCenter:			return TEXT("FrontCenter");
-		case LowFrequency:			return TEXT("LowFrequency");
-		case BackLeft:				return TEXT("BackLeft");
-		case BackRight:				return TEXT("BackRight");
-		case FrontLeftOfCenter:		return TEXT("FrontLeftOfCenter");
-		case FrontRightOfCenter:	return TEXT("FrontRightOfCenter");
-		case BackCenter:			return TEXT("BackCenter");
-		case SideLeft:				return TEXT("SideLeft");
-		case SideRight:				return TEXT("SideRight");
-		case TopCenter:				return TEXT("TopCenter");
-		case TopFrontLeft:			return TEXT("TopFrontLeft");
-		case TopFrontCenter:		return TEXT("TopFrontCenter");
-		case TopFrontRight:			return TEXT("TopFrontRight");
-		case TopBackLeft:			return TEXT("TopBackLeft");
-		case TopBackCenter:			return TEXT("TopBackCenter");
-		case TopBackRight:			return TEXT("TopBackRight");
-		case Unknown:				return TEXT("Unknown");
-
-		default:
-			return TEXT("UNSUPPORTED");
-		}
-	}
-}
 
 class FSoundWaveData;
 class FSoundWaveProxy;
diff --git a/Engine/Source/Runtime/AudioMixerCore/Public/AudioMixerChannel.h b/Engine/Source/Runtime/AudioMixerCore/Public/AudioMixerChannel.h
new file mode 100644
--- /dev/null
+++ b/Engine/Source/Runtime/AudioMixerCore/Public/AudioMixerChannel.h
@@ -0,0 +1,103 @@
+// Copyright Epic Games, Inc. All Rights Reserved.
+
+#pragma once
+
+#include "HAL/Platform.h"
+#include "Misc/Optional.h"
+#include "UObject/NameTypes.h"
+
+namespace EAudioMixerChannel
+{
+	/** Enumeration values represent sound file or speaker channel types. */
+	enum Type
+	{
+		FrontLeft,
+		FrontRight,
+		FrontCenter,
+		LowFrequency,
+		BackLeft,
+		BackRight,
+		FrontLeftOfCenter,
+		FrontRightOfCenter,
+		BackCenter,
+		SideLeft,
+		SideRight,
+		TopCenter,
+		TopFrontLeft,
+		TopFrontCenter,
+		TopFrontRight,
+		TopBackLeft,
+		TopBackCenter,
+		TopBackRight,
+		Unknown,
+		ChannelTypeCount,
+		DefaultChannel = FrontLeft
+	};
+
+// Foreach for each speaker (minus ChannelTypeCount and Default).
+// Not reflected enum, so do this by hand.
+#define FOREACH_EAUDIOMIXERCHANNEL(OP)\
+	OP(FrontLeft)\
+	OP(FrontRight)\
+	OP(FrontCenter)\
+	OP(LowFrequency)\
+	OP(BackLeft)\
+	OP(BackRight)\
+	OP(FrontLeftOfCenter)\
+	OP(FrontRightOfCenter)\
+	OP(BackCenter)\
+	OP(SideLeft)\
+	OP(SideRight)\
+	OP(TopCenter)\
+	OP(TopFrontLeft)\
+	OP(TopFrontCenter)\
+	OP(TopFrontRight)\
+	OP(TopBackLeft)\
+	OP(TopBackCenter)\
+	OP(TopBackRight)\
+	OP(Unknown)
+	
+	inline constexpr int32 MaxSupportedChannel = TopCenter;
+
+	constexpr const TCHAR* ToString(const Type InType) 
+	{
+		// Define using FOREACH macro 
+		#define CASE_TO_STRING(X) case X: return TEXT(#X);
+		switch (InType)
+		{
+			FOREACH_EAUDIOMIXERCHANNEL(CASE_TO_STRING)
+			default: break;
+		}
+		#undef CASE_TO_STRING
+		return TEXT("UNSUPPORTED");
+	}
+	
+	inline FName ToName(const Type InType)
+	{
+		// Keep these static.
+		#define CASE_TO_NAME(X) case X:\
+		{\
+			static FName Name(ToString(X));\
+			return Name;\
+		}
+		
+		// Define using FOREACH macro
+		switch (InType)
+		{
+			FOREACH_EAUDIOMIXERCHANNEL(CASE_TO_NAME)
+			default: break;
+		}
+		return TEXT("Unknown");
+	
+		#undef CASE_TO_NAME
+	}
+	
+	inline TOptional<Type> FromName(const FName InType)
+	{
+		// Define using FOREACH macro.
+		#define COMPARE_RETURN_NAME(X) if (InType == ToName(X)) return X;
+		FOREACH_EAUDIOMIXERCHANNEL(COMPARE_RETURN_NAME)
+		#undef COMPARE_RETURN_NAME
+		return {}; // Fail.
+	}
+}
