Fix for data race in underflow counters

#Synced-CL 50506470

diff --git a/Engine/Source/Runtime/AudioMixerCore/Private/AudioMixer.cpp b/Engine/Source/Runtime/AudioMixerCore/Private/AudioMixer.cpp
--- a/Engine/Source/Runtime/AudioMixerCore/Private/AudioMixer.cpp
+++ b/Engine/Source/Runtime/AudioMixerCore/Private/AudioMixer.cpp
@@ -606,10 +606,7 @@ namespace Audio
 			return;
 		}
 		
-		static int32 UnderrunCount = 0;
-		static int32 CurrentUnderrunCount = 0;
-		static uint64 TimeLastWarningCycles = 0;
-
+	
 		int32 NumSamplesPopped = 0;
 		TArrayView<const uint8> PoppedAudio = OutputBuffer.PopBufferData(NumSamplesPopped);
 
diff --git a/Engine/Source/Runtime/AudioMixerCore/Public/AudioMixer.h b/Engine/Source/Runtime/AudioMixerCore/Public/AudioMixer.h
--- a/Engine/Source/Runtime/AudioMixerCore/Public/AudioMixer.h
+++ b/Engine/Source/Runtime/AudioMixerCore/Public/AudioMixer.h
@@ -686,6 +686,9 @@ namespace Audio
 
 		/** Whether or not we warned of buffer underrun. */
 		bool bWarnedBufferUnderrun;
+		int32 UnderrunCount = 0;
+		int32 CurrentUnderrunCount = 0;
+		uint64 TimeLastWarningCycles = 0;
 
 		/** The audio render thread. */
 		//FRunnableThread* AudioRenderThread;
