[AudioInsights] - Tooltip for trace channels prefix with "AudioInsights"

#Synced-CL 50644355

diff --git a/Engine/Source/Runtime/AudioMixerCore/Private/AudioMixer.cpp b/Engine/Source/Runtime/AudioMixerCore/Private/AudioMixer.cpp
--- a/Engine/Source/Runtime/AudioMixerCore/Private/AudioMixer.cpp
+++ b/Engine/Source/Runtime/AudioMixerCore/Private/AudioMixer.cpp
@@ -25,8 +25,8 @@
 CSV_DEFINE_CATEGORY_MODULE(AUDIOMIXERCORE_API, Audio, false);
 
 #if UE_AUDIO_PROFILERTRACE_ENABLED
-UE_TRACE_CHANNEL_DEFINE(AudioChannel, "Audio logic data (such as asset load/unload/start/stop, etc.)");
-UE_TRACE_CHANNEL_DEFINE(AudioMixerChannel, "Audio parameter data (such as amplitude, pitch, etc.)");
+UE_TRACE_CHANNEL_DEFINE(AudioChannel, "AudioInsights logic data (such as asset load/unload/start/stop, etc.)");
+UE_TRACE_CHANNEL_DEFINE(AudioMixerChannel, "AudioInsights parameter data (such as amplitude, pitch, etc.)");
 #endif // UE_AUDIO_PROFILERTRACE_ENABLED
 
 // Audio Insights Event Log in-built event types
