[Audio] Fix deadlock between StopGeneratingAudio and FAudioRenderScheduler::RenderBlock during device swaps

During WASAPI/XAudio2 device swaps on Windows, CheckThreadedDeviceSwap() calls
StopGeneratingAudio() synchronously, which blocks on AudioRenderThread->Kill().
If the Audio Render Thread is inside ComputeNextBlockOfSamples ->
FAudioRenderScheduler::RenderBlock -> UE::Tasks::FTask::Wait(), this creates a
deadlock: the audio thread (a foreground worker) blocks on Kill(), starving the
UE Tasks workers that the render thread is waiting on to complete.

The fix splits StopGeneratingAudio into two phases during threaded device swaps:

Phase 1 (signal): When the async device swap completes, instead of synchronously
killing the render thread, we set StreamState to Stopping and trigger the render
event. CheckThreadedDeviceSwap returns false, releasing the audio thread back to
the task pool. This allows the UE Tasks workers to make progress, unblocking the
render thread's Task.Wait().

Phase 2 (cleanup): On the next tick, CheckThreadedDeviceSwap polls for the render
thread to have exited (StreamState == Stopped). Once confirmed, it performs the
cleanup (Kill, event pool returns) and proceeds with PostDeviceSwap. A timeout
fallback (au.OverrunTimeoutMSec) force-kills the thread if it fails to stop.

Force swap requests (e.g. OnSessionDisconnect, hardware stall) that arrive during
the wait are deferred and replayed after the current swap completes, rather than
being silently dropped. This prevents missing critical device change notifications
that could leave audio on an invalid output.

New methods added to IAudioMixerPlatformInterface:
- SignalStopGeneratingAudio(): Non-blocking signal to the render thread to stop.
- TryFinishStopGeneratingAudio(): Polls for render thread exit and cleans up.

Platform impact:
- Windows (WASAPI and XAudio2): Fixed. Both use CheckThreadedDeviceSwap via
the shared FAudioMixerPlatformSwappable base class.
- Consoles: Not affected. ShouldUseThreadedDeviceSwap() returns false on all
non-Windows platforms; they use the synchronous MoveAudioStreamToNewAudioDevice
path which is unchanged.

#Synced-CL 51668420

diff --git a/Engine/Source/Runtime/AudioMixerCore/Private/AudioMixer.cpp b/Engine/Source/Runtime/AudioMixerCore/Private/AudioMixer.cpp
--- a/Engine/Source/Runtime/AudioMixerCore/Private/AudioMixer.cpp
+++ b/Engine/Source/Runtime/AudioMixerCore/Private/AudioMixer.cpp
@@ -742,6 +742,100 @@ namespace Audio
 		bIsGeneratingAudio = false;
 	}
 
+	void IAudioMixerPlatformInterface::SignalStopGeneratingAudio()
+	{
+		SCOPED_NAMED_EVENT(IAudioMixerPlatformInterface_SignalStopGeneratingAudio, FColor::Blue);
+
+		if (AudioStreamInfo.StreamState != EAudioOutputStreamState::Stopped)
+		{
+			AudioStreamInfo.StreamState = EAudioOutputStreamState::Stopping;
+		}
+
+		if (AudioRenderEvent != nullptr)
+		{
+			// Make sure the thread wakes up so it can observe the Stopping state
+			AudioRenderEvent->Trigger();
+		}
+
+		StopGeneratingAudioSignalTime = FPlatformTime::Seconds();
+
+		UE_LOG(LogAudioMixer, Display, TEXT("SignalStopGeneratingAudio() called, InstanceID=%d"), InstanceID);
+	}
+
+	bool IAudioMixerPlatformInterface::TryFinishStopGeneratingAudio()
+	{
+		SCOPED_NAMED_EVENT(IAudioMixerPlatformInterface_TryFinishStopGeneratingAudio, FColor::Blue);
+
+		// Check if the render thread has exited on its own.
+		// RunInternal() sets StreamState = Stopped before returning.
+		const bool bRenderThreadStopped = (AudioStreamInfo.StreamState == EAudioOutputStreamState::Stopped);
+
+		if (!bRenderThreadStopped)
+		{
+			// Check for timeout - if the render thread hasn't stopped in time, we abandon it
+			// rather than calling Kill(true) which would block indefinitely (WaitForSingleObject
+			// with INFINITE timeout) and recreate the same deadlock we're trying to avoid.
+			const double TimeoutSeconds = static_cast<double>(FMath::Clamp<int32>(OverrunTimeoutCVar, 500, 5000)) / 1000.0;
+			const double ElapsedSeconds = FPlatformTime::Seconds() - StopGeneratingAudioSignalTime;
+			if (ElapsedSeconds < TimeoutSeconds)
+			{
+				return false; // Still waiting for the render thread to exit
+			}
+
+			UE_LOG(LogAudioMixer, Error,
+				TEXT("TryFinishStopGeneratingAudio() timed out after %.2fs waiting for render thread to stop. ")
+				TEXT("Abandoning render thread to avoid deadlock. The thread may still be blocked in Task.Wait(). InstanceID=%d"),
+				ElapsedSeconds, InstanceID);
+		}
+
+		if (AudioRenderThread.IsValid())
+		{
+			if (bRenderThreadStopped)
+			{
+				// Thread has exited - Kill returns immediately, just cleans up the handle.
+				SCOPED_NAMED_EVENT(IAudioMixerPlatformInterface_TryFinishStopGeneratingAudio_KillRenderThread, FColor::Blue);
+				AudioRenderThread->Kill(true);
+			}
+			else
+			{
+				// Thread is still running. We cannot call Kill(true) as it would block indefinitely.
+				// We cannot call Kill(false) as it closes the thread handle while the thread is still
+				// running, causing use-after-free when the thread eventually exits.
+				// Release ownership without destroying - the thread will leak but the process won't deadlock.
+				(void)AudioRenderThread.Release();
+				UE_LOG(LogAudioMixer, Warning, TEXT("TryFinishStopGeneratingAudio() leaked render thread handle to avoid deadlock. InstanceID=%d"), InstanceID);
+			}
+
+			AudioRenderThread.Reset();
+		}
+
+		AudioStreamInfo.StreamState = EAudioOutputStreamState::Stopped;
+
+		if (AudioRenderEvent != nullptr)
+		{
+			if (bRenderThreadStopped)
+			{
+				FPlatformProcess::ReturnSynchEventToPool(AudioRenderEvent);
+			}
+			// If the render thread is still running, it may still reference AudioRenderEvent.
+			// Don't return it to the pool - leak it to avoid use-after-free.
+			AudioRenderEvent = nullptr;
+		}
+
+		if (AudioFadeEvent != nullptr)
+		{
+			if (bRenderThreadStopped)
+			{
+				FPlatformProcess::ReturnSynchEventToPool(AudioFadeEvent);
+			}
+			// Same as above - leak rather than return to pool if thread is still running.
+			AudioFadeEvent = nullptr;
+		}
+
+		bIsGeneratingAudio = false;
+		return true;
+	}
+
 	void IAudioMixerPlatformInterface::Tick()
 	{
 		LLM_SCOPE(ELLMTag::AudioMixer);
@@ -863,6 +957,27 @@ namespace Audio
 		// This critical section protects that device swap context below.
 		FScopeLock Lock(&DeviceSwapCriticalSection);
 
+		// If the render thread is being stopped for a device swap, we cannot set StreamState to
+		// SwappingDevice (it would overwrite the Stopping state the render thread needs to see).
+		// Defer the request so it can be replayed after the current swap completes.
+		if (bWaitingForRenderThreadStop)
+		{
+			if (bInForce)
+			{
+				DeferredDeviceSwapId = DeviceID;
+				DeferredDeviceSwapReason = InReason ? InReason : TEXT("None specified");
+				UE_LOG(LogAudioMixer, Display, TEXT("Deferring force device swap request to '%s' (reason: '%s') until render thread stops."),
+					!DeviceID.IsEmpty() ? *DeviceID : TEXT("[System Default]"),
+					*DeferredDeviceSwapReason
+				);
+			}
+			else
+			{
+				UE_LOG(LogAudioMixer, Display, TEXT("Ignoring non-force device swap request, waiting for render thread to stop from previous swap."));
+			}
+			return false;
+		}
+
 		if (AllowDeviceSwap(bInForce))
 		{
 			UE_LOG(LogAudioMixer, Display, TEXT("Attempt to swap audio render device to new device: '%s', because: '%s', force=%d"),
@@ -887,7 +1002,7 @@ namespace Audio
 				(int32)bInForce
 			);
 		}
-		
+
 		return false;
 	}
 
@@ -921,8 +1036,8 @@ namespace Audio
 			UE_LOG(LogAudioMixer, Display, TEXT("Unable to device swap until backend is initialized."));
 			return false;
 		}
-		
-		// If the device is not running or if a device swap is already in progress, don't "double-trigger" a swap unless 
+
+		// If the device is not running or if a device swap is already in progress, don't "double-trigger" a swap unless
 		// bInForceSwap is true, in which case we'll replace the existing request.
 		if (ActiveDeviceSwap.IsValid() && !bInForceSwap)
 		{
@@ -955,50 +1070,86 @@ namespace Audio
 		bool bDidStopGeneratingAudio = false;
 		SCOPED_NAMED_EVENT(FAudioMixerPlatformSwappable_CheckThreadedDeviceSwap, FColor::Blue);
 
-		// Because this is called every tick, utilize the check-lock-check patter to avoid
+		// Phase 2: If we previously signaled the render thread to stop, check if it's done.
+		// This is checked first so that the swap completes as soon as the render thread exits.
+		if (bWaitingForRenderThreadStop)
+		{
+			FScopeLock Lock(&DeviceSwapCriticalSection);
+
+			if (!TryFinishStopGeneratingAudio())
+			{
+				return false; // Still waiting for render thread to exit
+			}
+
+			// Render thread has stopped — proceed with swap completion
+			bWaitingForRenderThreadStop = false;
+			bDidStopGeneratingAudio = true;
+
+			if (!PostDeviceSwap())
+			{
+				UE_LOG(LogAudioMixer, Warning, TEXT("FAudioMixerPlatformSwappable::CheckThreadedDeviceSwap PostDeviceSwap() failed"));
+			}
+
+			// If a force swap was requested while we were waiting, replay it now that
+			// the render thread is stopped and the previous swap is complete.
+			if (DeferredDeviceSwapId.IsSet())
+			{
+				FString ReplayDeviceId = MoveTemp(DeferredDeviceSwapId.GetValue());
+				FString ReplayReason = MoveTemp(DeferredDeviceSwapReason);
+				DeferredDeviceSwapId.Reset();
+
+				UE_LOG(LogAudioMixer, Display, TEXT("Replaying deferred device swap to '%s' (reason: '%s')"),
+					!ReplayDeviceId.IsEmpty() ? *ReplayDeviceId : TEXT("[System Default]"),
+					*ReplayReason);
+
+				// This is safe now — bWaitingForRenderThreadStop is false and StreamState is Stopped,
+				// so RequestDeviceSwap can set StreamState to SwappingDevice without conflict.
+				RequestDeviceSwap(ReplayDeviceId, /*bInForce*/ true, *ReplayReason);
+			}
+
+			return bDidStopGeneratingAudio;
+		}
+
+		// Phase 1: Check-lock-check pattern for starting or completing async swap.
+		// Because this is called every tick, utilize the check-lock-check pattern to avoid
 		// unnecessary locking.
 		if (AudioStreamInfo.StreamState == EAudioOutputStreamState::SwappingDevice || ActiveDeviceSwap.IsValid())
 		{
-			// Lock and verify this device is still running. 
+			// Lock and verify this device is still running.
 			FScopeLock Lock(&DeviceSwapCriticalSection);
-			
+
 			// Check the stream state again now that the lock is acquired in case the device is being closed down
 			EAudioOutputStreamState::Type StreamState = AudioStreamInfo.StreamState;
-			
+
 			// Start a job?
 			if (StreamState == EAudioOutputStreamState::SwappingDevice && !ActiveDeviceSwap.IsValid())
 			{
 				SCOPED_NAMED_EVENT(FAudioMixerPlatformSwappable_CheckThreadedDeviceSwap_StartAsyncSwap, FColor::Blue);
-				
+
 				if (!PreDeviceSwap())
 				{
 					UE_LOG(LogAudioMixer, Display, TEXT("FAudioMixerPlatformSwappable::CheckThreadedDeviceSwap PreDeviceSwap() failed"));
 					return bDidStopGeneratingAudio;
 				}
-				
+
 				EnqueueAsyncDeviceSwap();
 			}
-			else
+			else if (ActiveDeviceSwap.IsReady())
 			{
-				if (ActiveDeviceSwap.IsReady())  // Finished?
+				// Async swap finished — signal the render thread to stop without blocking.
+				// We return false this tick so that the audio thread releases back to the task pool,
+				// which breaks a potential deadlock where StopGeneratingAudio()->Kill() blocks on the
+				// render thread while the render thread is blocked in Task.Wait() waiting for task
+				// workers that can't run because this thread is blocked.
+				SCOPED_NAMED_EVENT(FAudioMixerPlatformSwappable_CheckThreadedDeviceSwap_SignalStop, FColor::Blue);
+
+				if (bIsUsingNullDevice)
 				{
-					SCOPED_NAMED_EVENT(FAudioMixerPlatformSwappable_CheckThreadedDeviceSwap_FinishedJobSwitchOver, FColor::Blue);
-					const FDeviceSwapResult* DeviceSwapResult = GetDeviceSwapResult();
-					
-					if (bIsUsingNullDevice)
-					{
-						StopRunningNullDevice();
-					}
-
-					StopGeneratingAudio();
-					bDidStopGeneratingAudio = true;
-
-					if (!PostDeviceSwap())
-					{
-						UE_LOG(LogAudioMixer, Warning, TEXT("FAudioMixerPlatformSwappable::CheckThreadedDeviceSwap PostDeviceSwap() failed"));
-						return bDidStopGeneratingAudio;
-					}
+					StopRunningNullDevice();
 				}
+
+				SignalStopGeneratingAudio();
+				bWaitingForRenderThreadStop = true;
 			}
 		}
 
diff --git a/Engine/Source/Runtime/AudioMixerCore/Public/AudioMixer.h b/Engine/Source/Runtime/AudioMixerCore/Public/AudioMixer.h
--- a/Engine/Source/Runtime/AudioMixerCore/Public/AudioMixer.h
+++ b/Engine/Source/Runtime/AudioMixerCore/Public/AudioMixer.h
@@ -656,6 +656,15 @@ namespace Audio
 		/** Stops the render thread from generating audio. */
 		AUDIOMIXERCORE_API void StopGeneratingAudio();
 
+		/** Non-blocking: signals the render thread to stop (sets Stopping state, triggers event)
+		 *  without waiting for the thread to exit. Use TryFinishStopGeneratingAudio() to complete cleanup. */
+		AUDIOMIXERCORE_API void SignalStopGeneratingAudio();
+
+		/** Checks if the render thread has stopped after SignalStopGeneratingAudio() was called.
+		 *  If the thread has stopped (or a timeout is reached), cleans up thread/events and returns true.
+		 *  If the thread is still running, returns false. */
+		AUDIOMIXERCORE_API bool TryFinishStopGeneratingAudio();
+
 		// Deprecated - use ApplyPrimaryAttenuation
 		UE_DEPRECATED(5.1, "ApplyMasterAttenuation is deprecated, please use ApplyPrimaryAttenuation instead.")
 		AUDIOMIXERCORE_API void ApplyMasterAttenuation(TArrayView<const uint8>& InOutPoppedAudio);
@@ -735,6 +744,9 @@ namespace Audio
 		FThreadSafeBool bIsUsingNullDevice;
 		FThreadSafeBool bIsGeneratingAudio;
 
+		/** Timestamp (seconds) when SignalStopGeneratingAudio was called, for timeout detection. */
+		double StopGeneratingAudioSignalTime = 0.0;
+
 		/** A Counter to provide the next unique id. */
 		AUDIOMIXERCORE_API static FThreadSafeCounter NextInstanceID;
 
@@ -810,6 +822,17 @@ namespace Audio
 		
 		/** Used to rate limit device swap attempts to ignore multiple requests */
 		double LastDeviceSwapTime = 0.0;
+
+		/** True when we've signaled the render thread to stop but haven't finished cleanup yet.
+		 *  Only read/written on the audio thread (inside CheckThreadedDeviceSwap). */
+		bool bWaitingForRenderThreadStop = false;
+
+		/** If a force device swap is requested while bWaitingForRenderThreadStop is true, the
+		 *  request is deferred here so it can be replayed after the current swap completes.
+		 *  Empty string means no deferred swap; non-empty or explicit default ("") is stored
+		 *  as a TOptional to distinguish "no request" from "request for default device". */
+		TOptional<FString> DeferredDeviceSwapId;
+		FString DeferredDeviceSwapReason;
 	};
 
 	class FScopedDeviceSwapLock final
